Copying missing files from RES to Main

diff --git a/Engine/Plugins/Online/OnlineSubsystemEOS/Source/Private/OnlineStoreEOS.cpp b/Engine/Plugins/Online/OnlineSubsystemEOS/Source/Private/OnlineStoreEOS.cpp
new file mode 100644
--- /dev/null
+++ b/Engine/Plugins/Online/OnlineSubsystemEOS/Source/Private/OnlineStoreEOS.cpp
@@ -0,0 +1,45 @@
+// Copyright Epic Games, Inc. All Rights Reserved.
+
+#include "OnlineStoreEOS.h"
+
+#if WITH_EOS_SDK
+
+#include "OnlineSubsystemEOS.h"
+#include "OnlineError.h"
+
+FOnlineStoreEOS::FOnlineStoreEOS(FOnlineSubsystemEOS* InSubsystem)
+	: EOSSubsystem(InSubsystem)
+{
+	check(EOSSubsystem != nullptr);
+}
+
+void FOnlineStoreEOS::QueryCategories(const FUniqueNetId& UserId, const FOnQueryOnlineStoreCategoriesComplete& Delegate)
+{
+}
+
+void FOnlineStoreEOS::GetCategories(TArray<FOnlineStoreCategory>& OutCategories) const
+{
+
+}
+
+void FOnlineStoreEOS::QueryOffersByFilter(const FUniqueNetId& UserId, const FOnlineStoreFilter& Filter, const FOnQueryOnlineStoreOffersComplete& Delegate)
+{
+
+}
+
+void FOnlineStoreEOS::QueryOffersById(const FUniqueNetId& UserId, const TArray<FUniqueOfferId>& OfferIds, const FOnQueryOnlineStoreOffersComplete& Delegate)
+{
+
+}
+
+void FOnlineStoreEOS::GetOffers(TArray<FOnlineStoreOfferRef>& OutOffers) const
+{
+
+}
+
+TSharedPtr<FOnlineStoreOffer> FOnlineStoreEOS::GetOffer(const FUniqueOfferId& OfferId) const
+{
+	return nullptr;
+}
+
+#endif
\ No newline at end of file
diff --git a/Engine/Plugins/Online/OnlineSubsystemEOS/Source/Private/OnlineStoreEOS.h b/Engine/Plugins/Online/OnlineSubsystemEOS/Source/Private/OnlineStoreEOS.h
new file mode 100644
--- /dev/null
+++ b/Engine/Plugins/Online/OnlineSubsystemEOS/Source/Private/OnlineStoreEOS.h
@@ -0,0 +1,49 @@
+// Copyright Epic Games, Inc. All Rights Reserved.
+
+#pragma once
+
+#include "CoreMinimal.h"
+#include "UObject/CoreOnline.h"
+#include "Interfaces/OnlineStoreInterfaceV2.h"
+#include "OnlineSubsystemEOSPackage.h"
+#include "OnlineSubsystemEOSTypes.h"
+
+#if WITH_EOS_SDK
+	#include "eos_ecom_types.h"
+
+/**
+ * Implementation for online store via EGS
+ */
+class FOnlineStoreEOS :
+	public IOnlineStoreV2,
+	public TSharedFromThis<FOnlineStoreEOS, ESPMode::ThreadSafe>
+{
+public:
+	virtual ~FOnlineStoreEOS() = default;
+
+// Begin IOnlineStoreV2
+	virtual void QueryCategories(const FUniqueNetId& UserId, const FOnQueryOnlineStoreCategoriesComplete& Delegate) override;
+	virtual void GetCategories(TArray<FOnlineStoreCategory>& OutCategories) const override;
+	virtual void QueryOffersByFilter(const FUniqueNetId& UserId, const FOnlineStoreFilter& Filter, const FOnQueryOnlineStoreOffersComplete& Delegate) override;
+	virtual void QueryOffersById(const FUniqueNetId& UserId, const TArray<FUniqueOfferId>& OfferIds, const FOnQueryOnlineStoreOffersComplete& Delegate) override;
+	virtual void GetOffers(TArray<FOnlineStoreOfferRef>& OutOffers) const override;
+	virtual TSharedPtr<FOnlineStoreOffer> GetOffer(const FUniqueOfferId& OfferId) const override;
+// End IOnlineStoreV2
+
+PACKAGE_SCOPE:
+	FOnlineStoreEOS(FOnlineSubsystemEOS* InSubsystem);
+
+private:
+	/** Default constructor disabled */
+	FOnlineStoreEOS() = delete;
+
+	/** Reference to the main EOS subsystem */
+	FOnlineSubsystemEOS* EOSSubsystem;
+
+	/** Map of offers queried via QueryOffersById */
+	TMap<FUniqueOfferId, FOnlineStoreOfferRef> CachedOffers;
+};
+
+typedef TSharedPtr<FOnlineStoreEOS, ESPMode::ThreadSafe> FOnlineStoreEOSPtr;
+
+#endif
