- Fix CIS errors by disabling deprecation warnings around calls to base class constructors

diff --git a/Engine/Plugins/Online/OnlineSubsystemEOS/Source/OnlineSubsystemEOS/Private/OnlineSubsystemEOSTypes.h b/Engine/Plugins/Online/OnlineSubsystemEOS/Source/OnlineSubsystemEOS/Private/OnlineSubsystemEOSTypes.h
--- a/Engine/Plugins/Online/OnlineSubsystemEOS/Source/OnlineSubsystemEOS/Private/OnlineSubsystemEOSTypes.h
+++ b/Engine/Plugins/Online/OnlineSubsystemEOS/Source/OnlineSubsystemEOS/Private/OnlineSubsystemEOSTypes.h
@@ -136,12 +136,16 @@ PACKAGE_SCOPE:
 	uint8 RawBytes[EOS_ID_BYTE_SIZE];
 private:
 	FUniqueNetIdEOS()
+		PRAGMA_DISABLE_DEPRECATION_WARNINGS
 		: FUniqueNetIdString(EMPTY_EASID EOS_ID_SEPARATOR EMPTY_PUID)
+		PRAGMA_ENABLE_DEPRECATION_WARNINGS
 	{
 	}
 	
 	explicit FUniqueNetIdEOS(uint8* Bytes, int32 Size)
+		PRAGMA_DISABLE_DEPRECATION_WARNINGS
 		: FUniqueNetIdString()
+		PRAGMA_ENABLE_DEPRECATION_WARNINGS
 	{
 		check(Size == EOS_ID_BYTE_SIZE);
 		EpicAccountIdStr = BytesToHex(Bytes, ID_HALF_BYTE_SIZE);
@@ -150,19 +154,25 @@ private:
 	}
 
 	explicit FUniqueNetIdEOS(const FString& InUniqueNetId)
+		PRAGMA_DISABLE_DEPRECATION_WARNINGS
 		: FUniqueNetIdString(InUniqueNetId)
+		PRAGMA_ENABLE_DEPRECATION_WARNINGS
 	{
 		ParseAccountIds();
 	}
 
 	explicit FUniqueNetIdEOS(FString&& InUniqueNetId)
+		PRAGMA_DISABLE_DEPRECATION_WARNINGS
 		: FUniqueNetIdString(MoveTemp(InUniqueNetId))
+		PRAGMA_ENABLE_DEPRECATION_WARNINGS
 	{
 		ParseAccountIds();
 	}
 
 	explicit FUniqueNetIdEOS(const FUniqueNetId& Src)
+		PRAGMA_DISABLE_DEPRECATION_WARNINGS
 		: FUniqueNetIdString(Src)
+		PRAGMA_ENABLE_DEPRECATION_WARNINGS
 	{
 		ParseAccountIds();
 	}
diff --git a/Engine/Plugins/Online/OnlineSubsystemEOS/Source/OnlineSubsystemEOSPlus/Private/OnlineUserEOSPlus.cpp b/Engine/Plugins/Online/OnlineSubsystemEOS/Source/OnlineSubsystemEOSPlus/Private/OnlineUserEOSPlus.cpp
--- a/Engine/Plugins/Online/OnlineSubsystemEOS/Source/OnlineSubsystemEOSPlus/Private/OnlineUserEOSPlus.cpp
+++ b/Engine/Plugins/Online/OnlineSubsystemEOS/Source/OnlineSubsystemEOSPlus/Private/OnlineUserEOSPlus.cpp
@@ -74,7 +74,9 @@ inline FString BuildEOSPlusStringId(FUniqueNetIdPtr InBaseUniqueNetId, FUniqueNe
 }
 
 FUniqueNetIdEOSPlus::FUniqueNetIdEOSPlus(FUniqueNetIdPtr InBaseUniqueNetId, FUniqueNetIdPtr InEOSUniqueNetId)
+	PRAGMA_DISABLE_DEPRECATION_WARNINGS
 	: FUniqueNetIdString(BuildEOSPlusStringId(InBaseUniqueNetId, InEOSUniqueNetId))
+	PRAGMA_ENABLE_DEPRECATION_WARNINGS
 	, BaseUniqueNetId(InBaseUniqueNetId)
 	, EOSUniqueNetId(InEOSUniqueNetId)
 {
