Adding logic to set product id, name and version in OSS EOS.

diff --git a/Engine/Plugins/Online/OnlineSubsystemEOS/Source/OnlineSubsystemEOS/Private/OnlineSubsystemEOS.cpp b/Engine/Plugins/Online/OnlineSubsystemEOS/Source/OnlineSubsystemEOS/Private/OnlineSubsystemEOS.cpp
--- a/Engine/Plugins/Online/OnlineSubsystemEOS/Source/OnlineSubsystemEOS/Private/OnlineSubsystemEOS.cpp
+++ b/Engine/Plugins/Online/OnlineSubsystemEOS/Source/OnlineSubsystemEOS/Private/OnlineSubsystemEOS.cpp
@@ -372,6 +372,19 @@ bool FOnlineSubsystemEOS::Init()
 	FString ErrorMessage;
 	SocketSubsystem->Init(ErrorMessage);
 
+	// We set the product id
+	FString ArtifactName;
+	FParse::Value(FCommandLine::Get(), TEXT("EpicApp="), ArtifactName);
+	FEOSArtifactSettings ArtifactSettings;
+	if (UEOSSettings::GetSettingsForArtifact(ArtifactName, ArtifactSettings))
+	{
+		ProductId = ArtifactSettings.ProductId;
+	}
+	else
+	{
+		UE_LOG_ONLINE(Warning, TEXT("[FOnlineSubsystemEOS::Init] Failed to find artifact settings object for artifact (%s). ProductIdAnsi not set."), *ArtifactName);
+	}
+
 	UserManager = MakeShareable(new FUserManagerEOS(this));
 	SessionInterfacePtr = MakeShareable(new FOnlineSessionEOS(this));
 	// Set the bucket id to use for all sessions based upon the name and version to avoid upgrade issues
diff --git a/Engine/Plugins/Online/OnlineSubsystemEOS/Source/OnlineSubsystemEOS/Private/UserManagerEOS.cpp b/Engine/Plugins/Online/OnlineSubsystemEOS/Source/OnlineSubsystemEOS/Private/UserManagerEOS.cpp
--- a/Engine/Plugins/Online/OnlineSubsystemEOS/Source/OnlineSubsystemEOS/Private/UserManagerEOS.cpp
+++ b/Engine/Plugins/Online/OnlineSubsystemEOS/Source/OnlineSubsystemEOS/Private/UserManagerEOS.cpp
@@ -9,6 +9,7 @@
 #include "SocketSubsystem.h"
 #include "OnlineError.h"
 #include "EOSSettings.h"
+#include "IEOSSDKManager.h"
 
 #include COMPILED_PLATFORM_HEADER(EOSHelpers.h)
 
@@ -2224,16 +2225,15 @@ void FUserManagerEOS::UpdatePresence(EOS_EpicAccountId AccountId)
 		}
 
 		FOnlineUserPresenceRef PresenceRef = NetIdStringToOnlineUserPresenceMap[NetId];
-		FString ProductId(UTF8_TO_TCHAR(PresenceInfo->ProductId));
-		FString ProdVersion(UTF8_TO_TCHAR(PresenceInfo->ProductVersion));
-		FString Platform(UTF8_TO_TCHAR(PresenceInfo->Platform));
+		const FString ProductId(UTF8_TO_TCHAR(PresenceInfo->ProductId));
+		const FString ProdVersion(UTF8_TO_TCHAR(PresenceInfo->ProductVersion));
+		const FString Platform(UTF8_TO_TCHAR(PresenceInfo->Platform));
 		// Convert the presence data to our format
 		PresenceRef->Status.State = ToEOnlinePresenceState(PresenceInfo->Status);
 		PresenceRef->Status.StatusStr = PresenceInfo->RichText;
 		PresenceRef->bIsOnline = PresenceRef->Status.State == EOnlinePresenceState::Online;
 		PresenceRef->bIsPlaying = !ProductId.IsEmpty();
-		PresenceRef->bIsPlayingThisGame = FCStringAnsi::Strcmp(PresenceInfo->ProductId, EOSSubsystem->ProductNameAnsi) == 0 &&
-			FCStringAnsi::Strcmp(PresenceInfo->ProductVersion, EOSSubsystem->ProductVersionAnsi) == 0;
+		PresenceRef->bIsPlayingThisGame = ProductId == EOSSubsystem->ProductId && ProdVersion == EOSSubsystem->EOSSDKManager->GetProductVersion();
 //		PresenceRef->bIsJoinable = ???;
 //		PresenceRef->bHasVoiceSupport = ???;
 		PresenceRef->Status.Properties.Add(TEXT("ProductId"), ProductId);
@@ -2425,7 +2425,7 @@ bool FUserManagerEOS::QueryUserIdMapping(const FUniqueNetId& UserId, const FStri
 	};
 
 	FQueryByDisplayNameOptions Options;
-	FCStringAnsi::Strncpy(Options.DisplayNameAnsi, TCHAR_TO_UTF8(*DisplayNameOrEmail), EOS_PRODUCTNAME_MAX_BUFFER_LEN);
+	FCStringAnsi::Strncpy(Options.DisplayNameAnsi, TCHAR_TO_UTF8(*DisplayNameOrEmail), EOS_OSS_STRING_BUFFER_LENGTH);
 	Options.LocalUserId = StringToAccountIdMap[NetId];
 	EOS_UserInfo_QueryUserInfoByDisplayName(EOSSubsystem->UserInfoHandle, &Options, CallbackObj, CallbackObj->GetCallbackPtr());
 
diff --git a/Engine/Plugins/Online/OnlineSubsystemEOS/Source/OnlineSubsystemEOS/Public/OnlineSubsystemEOS.h b/Engine/Plugins/Online/OnlineSubsystemEOS/Source/OnlineSubsystemEOS/Public/OnlineSubsystemEOS.h
--- a/Engine/Plugins/Online/OnlineSubsystemEOS/Source/OnlineSubsystemEOS/Public/OnlineSubsystemEOS.h
+++ b/Engine/Plugins/Online/OnlineSubsystemEOS/Source/OnlineSubsystemEOS/Public/OnlineSubsystemEOS.h
@@ -49,14 +49,6 @@ typedef TSharedPtr<class FOnlineUserCloudEOS, ESPMode::ThreadSafe> FOnlineUserCl
 
 typedef TSharedPtr<FPlatformEOSHelpers, ESPMode::ThreadSafe> FPlatformEOSHelpersPtr;
 
-#ifndef EOS_PRODUCTNAME_MAX_BUFFER_LEN
-	#define EOS_PRODUCTNAME_MAX_BUFFER_LEN 64
-#endif
-
-#ifndef EOS_PRODUCTVERSION_MAX_BUFFER_LEN
-	#define EOS_PRODUCTVERSION_MAX_BUFFER_LEN 64
-#endif
-
 /**
  *	OnlineSubsystemEOS - Implementation of the online subsystem for EOS services
  */
@@ -118,8 +110,7 @@ PACKAGE_SCOPE:
 	FOnlineSubsystemEOS() = delete;
 	explicit FOnlineSubsystemEOS(FName InInstanceName);
 
-	char ProductNameAnsi[EOS_PRODUCTNAME_MAX_BUFFER_LEN];
-	char ProductVersionAnsi[EOS_PRODUCTVERSION_MAX_BUFFER_LEN];
+	FString ProductId;
 
 	IEOSSDKManager* EOSSDKManager;
 
