Addition of UserCloud interface wrapper for EOSPlus

diff --git a/Engine/Plugins/Online/OnlineSubsystemEOS/Source/OnlineSubsystemEOSPlus/Private/OnlineSubsystemEOSPlus.cpp b/Engine/Plugins/Online/OnlineSubsystemEOS/Source/OnlineSubsystemEOSPlus/Private/OnlineSubsystemEOSPlus.cpp
--- a/Engine/Plugins/Online/OnlineSubsystemEOS/Source/OnlineSubsystemEOSPlus/Private/OnlineSubsystemEOSPlus.cpp
+++ b/Engine/Plugins/Online/OnlineSubsystemEOS/Source/OnlineSubsystemEOSPlus/Private/OnlineSubsystemEOSPlus.cpp
@@ -77,11 +77,13 @@ bool FOnlineSubsystemEOSPlus::Init()
 	LeaderboardsInterfacePtr = MakeShareable(new FOnlineLeaderboardsEOSPlus(this));
 	StoreInterfacePtr = MakeShareable(new FOnlineStoreEOSPlus(this));
 	ExternalUIInterfacePtr = MakeShareable(new FOnlineExternalUIEOSPlus(this));
+	UserCloudInterfacePtr = MakeShareable(new FOnlineUserCloudEOSPlus(this));
 
 	UserInterfacePtr->Initialize();
 	LeaderboardsInterfacePtr->Initialize();
 	StoreInterfacePtr->Initialize();
 	ExternalUIInterfacePtr->Initialize();
+	UserCloudInterfacePtr->Initialize();
 
 	return true;
 }
@@ -104,6 +106,7 @@ void FOnlineSubsystemEOSPlus::PreUnload()
 	DESTRUCT_INTERFACE(LeaderboardsInterfacePtr);
 	DESTRUCT_INTERFACE(StoreInterfacePtr);
 	DESTRUCT_INTERFACE(ExternalUIInterfacePtr);
+	DESTRUCT_INTERFACE(UserCloudInterfacePtr);
 
 #undef DESTRUCT_INTERFACE
 }
@@ -143,7 +146,7 @@ IOnlineSharedCloudPtr FOnlineSubsystemEOSPlus::GetSharedCloudInterface() const
 
 IOnlineUserCloudPtr FOnlineSubsystemEOSPlus::GetUserCloudInterface() const
 {
-	return nullptr;
+	return UserCloudInterfacePtr;
 }
 
 IOnlineEntitlementsPtr FOnlineSubsystemEOSPlus::GetEntitlementsInterface() const
diff --git a/Engine/Plugins/Online/OnlineSubsystemEOS/Source/OnlineSubsystemEOSPlus/Private/OnlineSubsystemEOSPlus.h b/Engine/Plugins/Online/OnlineSubsystemEOS/Source/OnlineSubsystemEOSPlus/Private/OnlineSubsystemEOSPlus.h
--- a/Engine/Plugins/Online/OnlineSubsystemEOS/Source/OnlineSubsystemEOSPlus/Private/OnlineSubsystemEOSPlus.h
+++ b/Engine/Plugins/Online/OnlineSubsystemEOS/Source/OnlineSubsystemEOSPlus/Private/OnlineSubsystemEOSPlus.h
@@ -13,6 +13,7 @@
 #include "OnlineLeaderboardsEOSPlus.h"
 #include "OnlineStoreEOSPlus.h"
 #include "OnlineExternalUIEOSPlus.h"
+#include "OnlineUserCloudEOSPlus.h"
 
 
 /**
@@ -84,6 +85,7 @@ PACKAGE_SCOPE:
 	FOnlineLeaderboardsEOSPlusPtr LeaderboardsInterfacePtr;
 	FOnlineStoreEOSPlusPtr StoreInterfacePtr;
 	FOnlineExternalUIEOSPlusPtr ExternalUIInterfacePtr;
+	FOnlineUserCloudEOSPlusPtr UserCloudInterfacePtr;
 };
 
 typedef TSharedPtr<FOnlineSubsystemEOSPlus, ESPMode::ThreadSafe> FOnlineSubsystemEOSPlusPtr;
diff --git a/Engine/Plugins/Online/OnlineSubsystemEOS/Source/OnlineSubsystemEOSPlus/Private/OnlineUserCloudEOSPlus.cpp b/Engine/Plugins/Online/OnlineSubsystemEOS/Source/OnlineSubsystemEOSPlus/Private/OnlineUserCloudEOSPlus.cpp
new file mode 100644
--- /dev/null
+++ b/Engine/Plugins/Online/OnlineSubsystemEOS/Source/OnlineSubsystemEOSPlus/Private/OnlineUserCloudEOSPlus.cpp
@@ -0,0 +1,421 @@
+// Copyright Epic Games, Inc. All Rights Reserved.
+
+#include "OnlineUserCloudEOSPlus.h"
+#include "OnlineSubsystemEOSPlus.h"
+#include "OnlineError.h"
+
+FOnlineUserCloudEOSPlus::FOnlineUserCloudEOSPlus(FOnlineSubsystemEOSPlus* InSubsystem)
+	: EOSPlus(InSubsystem)
+{
+	BaseUserCloudInterface = EOSPlus->BaseOSS->GetUserCloudInterface();
+}
+
+FOnlineUserCloudEOSPlus::~FOnlineUserCloudEOSPlus()
+{
+	if (BaseUserCloudInterface.IsValid())
+	{
+		BaseUserCloudInterface->ClearOnEnumerateUserFilesCompleteDelegates(this);
+		BaseUserCloudInterface->ClearOnWriteUserFileProgressDelegates(this);
+		BaseUserCloudInterface->ClearOnWriteUserFileCompleteDelegates(this);
+		BaseUserCloudInterface->ClearOnWriteUserFileCanceledDelegates(this);
+		BaseUserCloudInterface->ClearOnReadUserFileCompleteDelegates(this);
+		BaseUserCloudInterface->ClearOnDeleteUserFileCompleteDelegates(this);
+		BaseUserCloudInterface->ClearOnRequestUsageInfoCompleteDelegates(this);
+	}
+}
+
+FUniqueNetIdEOSPlusPtr FOnlineUserCloudEOSPlus::GetNetIdPlus(const FString& SourceId) const
+{
+	return EOSPlus->UserInterfacePtr->GetNetIdPlus(SourceId);
+}
+
+void FOnlineUserCloudEOSPlus::Initialize()
+{
+	if (BaseUserCloudInterface.IsValid())
+	{
+		BaseUserCloudInterface->AddOnEnumerateUserFilesCompleteDelegate_Handle(FOnEnumerateUserFilesCompleteDelegate::CreateThreadSafeSP(this, &FOnlineUserCloudEOSPlus::OnEnumerateUserFilesComplete));
+		BaseUserCloudInterface->AddOnWriteUserFileProgressDelegate_Handle(FOnWriteUserFileProgressDelegate::CreateThreadSafeSP(this, &FOnlineUserCloudEOSPlus::OnWriteUserFileProgress));
+		BaseUserCloudInterface->AddOnWriteUserFileCompleteDelegate_Handle(FOnWriteUserFileCompleteDelegate::CreateThreadSafeSP(this, &FOnlineUserCloudEOSPlus::OnWriteUserFileComplete));
+		BaseUserCloudInterface->AddOnWriteUserFileCanceledDelegate_Handle(FOnWriteUserFileCanceledDelegate::CreateThreadSafeSP(this, &FOnlineUserCloudEOSPlus::OnWriteUserFileCanceled));
+		BaseUserCloudInterface->AddOnReadUserFileCompleteDelegate_Handle(FOnReadUserFileCompleteDelegate::CreateThreadSafeSP(this, &FOnlineUserCloudEOSPlus::OnReadUserFileComplete));
+		BaseUserCloudInterface->AddOnDeleteUserFileCompleteDelegate_Handle(FOnDeleteUserFileCompleteDelegate::CreateThreadSafeSP(this, &FOnlineUserCloudEOSPlus::OnDeleteUserFileComplete));
+		BaseUserCloudInterface->AddOnRequestUsageInfoCompleteDelegate_Handle(FOnRequestUsageInfoCompleteDelegate::CreateThreadSafeSP(this, &FOnlineUserCloudEOSPlus::OnRequestUsageInfoComplete));
+	}
+	else
+	{
+		UE_LOG_ONLINE(Warning, TEXT("[FOnlineUserCloudEOSPlus::Initialize] BaseUserCloudInterface delegates not bound. Base interface not valid"));
+	}
+}
+
+//~ Begin IOnlineUserCloud Interface
+
+bool FOnlineUserCloudEOSPlus::GetFileContents(const FUniqueNetId& UserId, const FString& FileName, TArray<uint8>& FileContents)
+{
+	bool bResult = false;
+
+	const FUniqueNetIdEOSPlusPtr NetIdPlus = GetNetIdPlus(UserId.ToString());
+	if (NetIdPlus.IsValid())
+	{
+		const bool bIsBaseNetIdValid = ensure(NetIdPlus->GetBaseNetId().IsValid());
+		const bool bIsBaseUserCloudInterfaceValid = BaseUserCloudInterface.IsValid();
+		if (bIsBaseNetIdValid && bIsBaseUserCloudInterfaceValid)
+		{
+			bResult = BaseUserCloudInterface->GetFileContents(*NetIdPlus->GetBaseNetId(), FileName, FileContents);
+		}
+		else
+		{
+			UE_LOG_ONLINE(Warning, TEXT("[FOnlineUserCloudEOSPlus::GetFileContents] Unable to call method in base interface. IsBaseNetIdValid=%s IsBaseUserCloudInterfaceValid=%s."), *LexToString(bIsBaseNetIdValid), *LexToString(bIsBaseUserCloudInterfaceValid));
+		}
+	}
+	else
+	{
+		UE_LOG_ONLINE(Warning, TEXT("[FOnlineUserCloudEOSPlus::GetFileContents] Unable to call method in base interface. Unknown user (%s)"), *UserId.ToString());
+	}
+
+	return bResult;
+}
+
+bool FOnlineUserCloudEOSPlus::ClearFiles(const FUniqueNetId& UserId)
+{
+	bool bResult = false;
+
+	const FUniqueNetIdEOSPlusPtr NetIdPlus = GetNetIdPlus(UserId.ToString());
+	if (NetIdPlus.IsValid())
+	{
+		const bool bIsBaseNetIdValid = ensure(NetIdPlus->GetBaseNetId().IsValid());
+		const bool bIsBaseUserCloudInterfaceValid = BaseUserCloudInterface.IsValid();
+		if (bIsBaseNetIdValid && bIsBaseUserCloudInterfaceValid)
+		{
+			bResult = BaseUserCloudInterface->ClearFiles(*NetIdPlus->GetBaseNetId());
+		}
+		else
+		{
+			UE_LOG_ONLINE(Warning, TEXT("[FOnlineUserCloudEOSPlus::ClearFiles] Unable to call method in base interface. IsBaseNetIdValid=%s IsBaseUserCloudInterfaceValid=%s."), *LexToString(bIsBaseNetIdValid), *LexToString(bIsBaseUserCloudInterfaceValid));
+		}
+	}
+	else
+	{
+		UE_LOG_ONLINE(Warning, TEXT("[FOnlineUserCloudEOSPlus::ClearFiles] Unable to call method in base interface. Unknown user (%s)"), *UserId.ToString());
+	}
+
+	return bResult;
+}
+
+bool FOnlineUserCloudEOSPlus::ClearFile(const FUniqueNetId& UserId, const FString& FileName)
+{
+	bool bResult = false;
+
+	const FUniqueNetIdEOSPlusPtr NetIdPlus = GetNetIdPlus(UserId.ToString());
+	if (NetIdPlus.IsValid())
+	{
+		const bool bIsBaseNetIdValid = ensure(NetIdPlus->GetBaseNetId().IsValid());
+		const bool bIsBaseUserCloudInterfaceValid = BaseUserCloudInterface.IsValid();
+		if (bIsBaseNetIdValid && bIsBaseUserCloudInterfaceValid)
+		{
+			bResult = BaseUserCloudInterface->ClearFile(*NetIdPlus->GetBaseNetId(), FileName);
+		}
+		else
+		{
+			UE_LOG_ONLINE(Warning, TEXT("[FOnlineUserCloudEOSPlus::ClearFile] Unable to call method in base interface. IsBaseNetIdValid=%s IsBaseUserCloudInterfaceValid=%s."), *LexToString(bIsBaseNetIdValid), *LexToString(bIsBaseUserCloudInterfaceValid));
+		}
+	}
+	else
+	{
+		UE_LOG_ONLINE(Warning, TEXT("[FOnlineUserCloudEOSPlus::ClearFile] Unable to call method in base interface. Unknown user (%s)"), *UserId.ToString());
+	}
+
+	return bResult;
+}
+
+void FOnlineUserCloudEOSPlus::EnumerateUserFiles(const FUniqueNetId& UserId)
+{
+	const FUniqueNetIdEOSPlusPtr NetIdPlus = GetNetIdPlus(UserId.ToString());
+	if (NetIdPlus.IsValid())
+	{
+		const bool bIsBaseNetIdValid = ensure(NetIdPlus->GetBaseNetId().IsValid());
+		const bool bIsBaseUserCloudInterfaceValid = BaseUserCloudInterface.IsValid();
+		if (bIsBaseNetIdValid && bIsBaseUserCloudInterfaceValid)
+		{
+			BaseUserCloudInterface->EnumerateUserFiles(*NetIdPlus->GetBaseNetId());
+		}
+		else
+		{
+			UE_LOG_ONLINE(Warning, TEXT("[FOnlineUserCloudEOSPlus::EnumerateUserFiles] Unable to call method in base interface. IsBaseNetIdValid=%s IsBaseUserCloudInterfaceValid=%s."), *LexToString(bIsBaseNetIdValid), *LexToString(bIsBaseUserCloudInterfaceValid));
+		}
+	}
+	else
+	{
+		UE_LOG_ONLINE(Warning, TEXT("[FOnlineUserCloudEOSPlus::EnumerateUserFiles] Unable to call method in base interface. Unknown user (%s)"), *UserId.ToString());
+	}
+}
+	
+void FOnlineUserCloudEOSPlus::GetUserFileList(const FUniqueNetId& UserId, TArray<FCloudFileHeader>& UserFiles)
+{
+	const FUniqueNetIdEOSPlusPtr NetIdPlus = GetNetIdPlus(UserId.ToString());
+	if (NetIdPlus.IsValid())
+	{
+		const bool bIsBaseNetIdValid = ensure(NetIdPlus->GetBaseNetId().IsValid());
+		const bool bIsBaseUserCloudInterfaceValid = BaseUserCloudInterface.IsValid();
+		if (bIsBaseNetIdValid && bIsBaseUserCloudInterfaceValid)
+		{
+			BaseUserCloudInterface->GetUserFileList(*NetIdPlus->GetBaseNetId(), UserFiles);
+		}
+		else
+		{
+			UE_LOG_ONLINE(Warning, TEXT("[FOnlineUserCloudEOSPlus::GetUserFileList] Unable to call method in base interface. IsBaseNetIdValid=%s IsBaseUserCloudInterfaceValid=%s."), *LexToString(bIsBaseNetIdValid), *LexToString(bIsBaseUserCloudInterfaceValid));
+		}
+	}
+	else
+	{
+		UE_LOG_ONLINE(Warning, TEXT("[FOnlineUserCloudEOSPlus::GetUserFileList] Unable to call method in base interface. Unknown user (%s)"), *UserId.ToString());
+	}
+}
+	
+bool FOnlineUserCloudEOSPlus::ReadUserFile(const FUniqueNetId& UserId, const FString& FileName)
+{
+	bool bResult = false;
+
+	const FUniqueNetIdEOSPlusPtr NetIdPlus = GetNetIdPlus(UserId.ToString());
+	if (NetIdPlus.IsValid())
+	{
+		const bool bIsBaseNetIdValid = ensure(NetIdPlus->GetBaseNetId().IsValid());
+		const bool bIsBaseUserCloudInterfaceValid = BaseUserCloudInterface.IsValid();
+		if (bIsBaseNetIdValid && bIsBaseUserCloudInterfaceValid)
+		{
+			bResult = BaseUserCloudInterface->ReadUserFile(*NetIdPlus->GetBaseNetId(), FileName);
+		}
+		else
+		{
+			UE_LOG_ONLINE(Warning, TEXT("[FOnlineUserCloudEOSPlus::ReadUserFile] Unable to call method in base interface. IsBaseNetIdValid=%s IsBaseUserCloudInterfaceValid=%s."), *LexToString(bIsBaseNetIdValid), *LexToString(bIsBaseUserCloudInterfaceValid));
+		}
+	}
+	else
+	{
+		UE_LOG_ONLINE(Warning, TEXT("[FOnlineUserCloudEOSPlus::ReadUserFile] Unable to call method in base interface. Unknown user (%s)"), *UserId.ToString());
+	}
+
+	return bResult;
+}
+	
+bool FOnlineUserCloudEOSPlus::WriteUserFile(const FUniqueNetId& UserId, const FString& FileName, TArray<uint8>& FileContents, bool bCompressBeforeUpload)
+{
+	bool bResult = false;
+
+	const FUniqueNetIdEOSPlusPtr NetIdPlus = GetNetIdPlus(UserId.ToString());
+	if (NetIdPlus.IsValid())
+	{
+		const bool bIsBaseNetIdValid = ensure(NetIdPlus->GetBaseNetId().IsValid());
+		const bool bIsBaseUserCloudInterfaceValid = BaseUserCloudInterface.IsValid();
+		if (bIsBaseNetIdValid && bIsBaseUserCloudInterfaceValid)
+		{
+			bResult = BaseUserCloudInterface->WriteUserFile(*NetIdPlus->GetBaseNetId(), FileName, FileContents, bCompressBeforeUpload);
+		}
+		else
+		{
+			UE_LOG_ONLINE(Warning, TEXT("[FOnlineUserCloudEOSPlus::WriteUserFile] Unable to call method in base interface. IsBaseNetIdValid=%s IsBaseUserCloudInterfaceValid=%s."), *LexToString(bIsBaseNetIdValid), *LexToString(bIsBaseUserCloudInterfaceValid));
+		}
+	}
+	else
+	{
+		UE_LOG_ONLINE(Warning, TEXT("[FOnlineUserCloudEOSPlus::WriteUserFile] Unable to call method in base interface. Unknown user (%s)"), *UserId.ToString());
+	}
+
+	return bResult;
+}
+	
+void FOnlineUserCloudEOSPlus::CancelWriteUserFile(const FUniqueNetId& UserId, const FString& FileName)
+{
+	const FUniqueNetIdEOSPlusPtr NetIdPlus = GetNetIdPlus(UserId.ToString());
+	if (NetIdPlus.IsValid())
+	{
+		const bool bIsBaseNetIdValid = ensure(NetIdPlus->GetBaseNetId().IsValid());
+		const bool bIsBaseUserCloudInterfaceValid = BaseUserCloudInterface.IsValid();
+		if (bIsBaseNetIdValid && bIsBaseUserCloudInterfaceValid)
+		{
+			BaseUserCloudInterface->CancelWriteUserFile(*NetIdPlus->GetBaseNetId(), FileName);
+		}
+		else
+		{
+			UE_LOG_ONLINE(Warning, TEXT("[FOnlineUserCloudEOSPlus::CancelWriteUserFile] Unable to call method in base interface. IsBaseNetIdValid=%s IsBaseUserCloudInterfaceValid=%s."), *LexToString(bIsBaseNetIdValid), *LexToString(bIsBaseUserCloudInterfaceValid));
+		}
+	}
+	else
+	{
+		UE_LOG_ONLINE(Warning, TEXT("[FOnlineUserCloudEOSPlus::CancelWriteUserFile] Unable to call method in base interface. Unknown user (%s)"), *UserId.ToString());
+	}
+}
+
+bool FOnlineUserCloudEOSPlus::DeleteUserFile(const FUniqueNetId& UserId, const FString& FileName, bool bShouldCloudDelete, bool bShouldLocallyDelete)
+{
+	bool bResult = false;
+
+	const FUniqueNetIdEOSPlusPtr NetIdPlus = GetNetIdPlus(UserId.ToString());
+	if (NetIdPlus.IsValid())
+	{
+		const bool bIsBaseNetIdValid = ensure(NetIdPlus->GetBaseNetId().IsValid());
+		const bool bIsBaseUserCloudInterfaceValid = BaseUserCloudInterface.IsValid();
+		if (bIsBaseNetIdValid && bIsBaseUserCloudInterfaceValid)
+		{
+			bResult = BaseUserCloudInterface->DeleteUserFile(*NetIdPlus->GetBaseNetId(), FileName, bShouldCloudDelete, bShouldLocallyDelete);
+		}
+		else
+		{
+			UE_LOG_ONLINE(Warning, TEXT("[FOnlineUserCloudEOSPlus::DeleteUserFile] Unable to call method in base interface. IsBaseNetIdValid=%s IsBaseUserCloudInterfaceValid=%s."), *LexToString(bIsBaseNetIdValid), *LexToString(bIsBaseUserCloudInterfaceValid));
+		}
+	}
+	else
+	{
+		UE_LOG_ONLINE(Warning, TEXT("[FOnlineUserCloudEOSPlus::DeleteUserFile] Unable to call method in base interface. Unknown user (%s)"), *UserId.ToString());
+	}
+
+	return bResult;
+}
+
+bool FOnlineUserCloudEOSPlus::RequestUsageInfo(const FUniqueNetId& UserId)
+{
+	bool bResult = false;
+
+	const FUniqueNetIdEOSPlusPtr NetIdPlus = GetNetIdPlus(UserId.ToString());
+	if (NetIdPlus.IsValid())
+	{
+		const bool bIsBaseNetIdValid = ensure(NetIdPlus->GetBaseNetId().IsValid());
+		const bool bIsBaseUserCloudInterfaceValid = BaseUserCloudInterface.IsValid();
+		if (bIsBaseNetIdValid && bIsBaseUserCloudInterfaceValid)
+		{
+			bResult = BaseUserCloudInterface->RequestUsageInfo(*NetIdPlus->GetBaseNetId());
+		}
+		else
+		{
+			UE_LOG_ONLINE(Warning, TEXT("[FOnlineUserCloudEOSPlus::RequestUsageInfo] Unable to call method in base interface. IsBaseNetIdValid=%s IsBaseUserCloudInterfaceValid=%s."), *LexToString(bIsBaseNetIdValid), *LexToString(bIsBaseUserCloudInterfaceValid));
+		}
+	}
+	else
+	{
+		UE_LOG_ONLINE(Warning, TEXT("[FOnlineUserCloudEOSPlus::RequestUsageInfo] Unable to call method in base interface. Unknown user (%s)"), *UserId.ToString());
+	}
+
+	return bResult;
+}
+
+void FOnlineUserCloudEOSPlus::DumpCloudState(const FUniqueNetId& UserId)
+{
+	const FUniqueNetIdEOSPlusPtr NetIdPlus = GetNetIdPlus(UserId.ToString());
+	if (NetIdPlus.IsValid())
+	{
+		const bool bIsBaseNetIdValid = ensure(NetIdPlus->GetBaseNetId().IsValid());
+		const bool bIsBaseUserCloudInterfaceValid = BaseUserCloudInterface.IsValid();
+		if (bIsBaseNetIdValid && bIsBaseUserCloudInterfaceValid)
+		{
+			BaseUserCloudInterface->DumpCloudState(*NetIdPlus->GetBaseNetId());
+		}
+		else
+		{
+			UE_LOG_ONLINE(Warning, TEXT("[FOnlineUserCloudEOSPlus::DumpCloudState] Unable to call method in base interface. IsBaseNetIdValid=%s IsBaseUserCloudInterfaceValid=%s."), *LexToString(bIsBaseNetIdValid), *LexToString(bIsBaseUserCloudInterfaceValid));
+		}
+	}
+	else
+	{
+		UE_LOG_ONLINE(Warning, TEXT("[FOnlineUserCloudEOSPlus::DumpCloudState] Unable to call method in base interface. Unknown user (%s)"), *UserId.ToString());
+	}
+}
+
+void FOnlineUserCloudEOSPlus::DumpCloudFileState(const FUniqueNetId& UserId, const FString& FileName)
+{
+	const FUniqueNetIdEOSPlusPtr NetIdPlus = GetNetIdPlus(UserId.ToString());
+	if (NetIdPlus.IsValid())
+	{
+		const bool bIsBaseNetIdValid = ensure(NetIdPlus->GetBaseNetId().IsValid());
+		const bool bIsBaseUserCloudInterfaceValid = BaseUserCloudInterface.IsValid();
+		if (bIsBaseNetIdValid && bIsBaseUserCloudInterfaceValid)
+		{
+			BaseUserCloudInterface->DumpCloudFileState(*NetIdPlus->GetBaseNetId(), FileName);
+		}
+		else
+		{
+			UE_LOG_ONLINE(Warning, TEXT("[FOnlineUserCloudEOSPlus::DumpCloudFileState] Unable to call method in base interface. IsBaseNetIdValid=%s IsBaseUserCloudInterfaceValid=%s."), *LexToString(bIsBaseNetIdValid), *LexToString(bIsBaseUserCloudInterfaceValid));
+		}
+	}
+	else
+	{
+		UE_LOG_ONLINE(Warning, TEXT("[FOnlineUserCloudEOSPlus::DumpCloudFileState] Unable to call method in base interface. Unknown user (%s)"), *UserId.ToString());
+	}
+}
+
+void FOnlineUserCloudEOSPlus::OnEnumerateUserFilesComplete(bool bWasSuccessful, const FUniqueNetId& UserId)
+{
+	const FUniqueNetIdEOSPlusPtr NetIdPlus = GetNetIdPlus(UserId.ToString());
+	if (!NetIdPlus.IsValid())
+	{
+		UE_LOG_ONLINE(Warning, TEXT("[FOnlineUserCloudEOSPlus::OnEnumerateUserFilesComplete] Unknown user (%s)"), *UserId.ToString());
+	}
+
+	TriggerOnEnumerateUserFilesCompleteDelegates(bWasSuccessful, *NetIdPlus);
+}
+
+void FOnlineUserCloudEOSPlus::OnWriteUserFileProgress(int32 BytesWritten, const FUniqueNetId& UserId, const FString& FileName)
+{
+	const FUniqueNetIdEOSPlusPtr NetIdPlus = GetNetIdPlus(UserId.ToString());
+	if (!NetIdPlus.IsValid())
+	{
+		UE_LOG_ONLINE(Warning, TEXT("[FOnlineUserCloudEOSPlus::OnWriteUserFileProgress] Unknown user (%s)"), *UserId.ToString());
+	}
+
+	TriggerOnWriteUserFileProgressDelegates(BytesWritten, *NetIdPlus, FileName);
+}
+
+void FOnlineUserCloudEOSPlus::OnWriteUserFileComplete(bool bWasSuccessful, const FUniqueNetId& UserId, const FString& FileName)
+{
+	const FUniqueNetIdEOSPlusPtr NetIdPlus = GetNetIdPlus(UserId.ToString());
+	if (!NetIdPlus.IsValid())
+	{
+		UE_LOG_ONLINE(Warning, TEXT("[FOnlineUserCloudEOSPlus::OnWriteUserFileComplete] Unknown user (%s)"), *UserId.ToString());
+	}
+
+	TriggerOnWriteUserFileCompleteDelegates(bWasSuccessful, *NetIdPlus, FileName);
+}
+
+void FOnlineUserCloudEOSPlus::OnWriteUserFileCanceled(bool bWasSuccessful, const FUniqueNetId& UserId, const FString& FileName)
+{
+	const FUniqueNetIdEOSPlusPtr NetIdPlus = GetNetIdPlus(UserId.ToString());
+	if (!NetIdPlus.IsValid())
+	{
+		UE_LOG_ONLINE(Warning, TEXT("[FOnlineUserCloudEOSPlus::OnWriteUserFileCanceled] Unknown user (%s)"), *UserId.ToString());
+	}
+
+	TriggerOnWriteUserFileCanceledDelegates(bWasSuccessful, *NetIdPlus, FileName);
+}
+
+void FOnlineUserCloudEOSPlus::OnReadUserFileComplete(bool bWasSuccessful, const FUniqueNetId& UserId, const FString& FileName)
+{
+	const FUniqueNetIdEOSPlusPtr NetIdPlus = GetNetIdPlus(UserId.ToString());
+	if (!NetIdPlus.IsValid())
+	{
+		UE_LOG_ONLINE(Warning, TEXT("[FOnlineUserCloudEOSPlus::OnReadUserFileComplete] Unknown user (%s)"), *UserId.ToString());
+	}
+
+	TriggerOnReadUserFileCompleteDelegates(bWasSuccessful, *NetIdPlus, FileName);
+}
+
+void FOnlineUserCloudEOSPlus::OnDeleteUserFileComplete(bool bWasSuccessful, const FUniqueNetId& UserId, const FString& FileName)
+{
+	const FUniqueNetIdEOSPlusPtr NetIdPlus = GetNetIdPlus(UserId.ToString());
+	if (!NetIdPlus.IsValid())
+	{
+		UE_LOG_ONLINE(Warning, TEXT("[FOnlineUserCloudEOSPlus::OnDeleteUserFileComplete] Unknown user (%s)"), *UserId.ToString());
+	}
+
+	TriggerOnDeleteUserFileCompleteDelegates(bWasSuccessful, *NetIdPlus, FileName);
+}
+	
+void FOnlineUserCloudEOSPlus::OnRequestUsageInfoComplete(bool bWasSuccessful, const FUniqueNetId& UserId, int64 BytesUsed, const TOptional<int64>& TotalQuota)
+{
+	const FUniqueNetIdEOSPlusPtr NetIdPlus = GetNetIdPlus(UserId.ToString());
+	if (!NetIdPlus.IsValid())
+	{
+		UE_LOG_ONLINE(Warning, TEXT("[FOnlineUserCloudEOSPlus::OnRequestUsageInfoComplete] Unknown user (%s)"), *UserId.ToString());
+	}
+
+	TriggerOnRequestUsageInfoCompleteDelegates(bWasSuccessful, *NetIdPlus, BytesUsed, TotalQuota);
+}
+	
+//~ End IOnlineUserCloud Interface
\ No newline at end of file
diff --git a/Engine/Plugins/Online/OnlineSubsystemEOS/Source/OnlineSubsystemEOSPlus/Private/OnlineUserCloudEOSPlus.h b/Engine/Plugins/Online/OnlineSubsystemEOS/Source/OnlineSubsystemEOSPlus/Private/OnlineUserCloudEOSPlus.h
new file mode 100644
--- /dev/null
+++ b/Engine/Plugins/Online/OnlineSubsystemEOS/Source/OnlineSubsystemEOSPlus/Private/OnlineUserCloudEOSPlus.h
@@ -0,0 +1,62 @@
+// Copyright Epic Games, Inc. All Rights Reserved.
+
+#pragma once
+
+#include "CoreMinimal.h"
+#include "Online/CoreOnline.h"
+#include "OnlineSubsystem.h"
+#include "Interfaces/OnlineUserCloudInterface.h"
+#include "OnlineUserEOSPlus.h"
+
+class FOnlineSubsystemEOSPlus;
+
+/**
+ * Interface for encapsulating the platform user cloud interface
+ */
+class FOnlineUserCloudEOSPlus :
+	public IOnlineUserCloud,
+	public TSharedFromThis<FOnlineUserCloudEOSPlus, ESPMode::ThreadSafe>
+{
+public:
+	FOnlineUserCloudEOSPlus() = delete;
+	virtual ~FOnlineUserCloudEOSPlus();
+
+	void Initialize();
+
+	//~ Begin IOnlineUserCloud Interface
+	virtual bool GetFileContents(const FUniqueNetId& UserId, const FString& FileName, TArray<uint8>& FileContents) override;
+	virtual bool ClearFiles(const FUniqueNetId& UserId) override;
+	virtual bool ClearFile(const FUniqueNetId& UserId, const FString& FileName) override;
+	virtual void EnumerateUserFiles(const FUniqueNetId& UserId) override;
+	virtual void GetUserFileList(const FUniqueNetId& UserId, TArray<FCloudFileHeader>& UserFiles) override;
+	virtual bool ReadUserFile(const FUniqueNetId& UserId, const FString& FileName) override;
+	virtual bool WriteUserFile(const FUniqueNetId& UserId, const FString& FileName, TArray<uint8>& FileContents, bool bCompressBeforeUpload = false) override;
+	virtual void CancelWriteUserFile(const FUniqueNetId& UserId, const FString& FileName) override;
+	virtual bool DeleteUserFile(const FUniqueNetId& UserId, const FString& FileName, bool bShouldCloudDelete, bool bShouldLocallyDelete) override;
+	virtual bool RequestUsageInfo(const FUniqueNetId& UserId) override;
+	virtual void DumpCloudState(const FUniqueNetId& UserId) override;
+	virtual void DumpCloudFileState(const FUniqueNetId& UserId, const FString& FileName) override;
+	//~ End IOnlineUserCloud Interface
+
+PACKAGE_SCOPE:
+	FOnlineUserCloudEOSPlus(FOnlineSubsystemEOSPlus* InSubsystem);
+
+	void OnEnumerateUserFilesComplete(bool bWasSuccessful, const FUniqueNetId& UserId);
+	void OnWriteUserFileProgress(int32 BytesWritten, const FUniqueNetId& UserId, const FString& FileName);
+	void OnWriteUserFileComplete(bool bWasSuccessful, const FUniqueNetId& UserId, const FString& FileName);
+	void OnWriteUserFileCanceled(bool bWasSuccessful, const FUniqueNetId& UserId, const FString& FileName);
+	void OnReadUserFileComplete(bool bWasSuccessful, const FUniqueNetId& UserId, const FString& FileName);
+	void OnDeleteUserFileComplete(bool bWasSuccessful, const FUniqueNetId& UserId, const FString& FileName);
+	void OnRequestUsageInfoComplete(bool bWasSuccessful, const FUniqueNetId& UserId, int64 BytesUsed, const TOptional<int64>& TotalQuota);
+
+private:
+	FUniqueNetIdEOSPlusPtr GetNetIdPlus(const FString& SourceId) const;
+
+	/** Reference to the owning EOS plus subsystem */
+	FOnlineSubsystemEOSPlus* EOSPlus;
+	
+	// We don't support EOS mirroring yet
+	IOnlineUserCloudPtr BaseUserCloudInterface;
+};
+
+typedef TSharedPtr<FOnlineUserCloudEOSPlus, ESPMode::ThreadSafe> FOnlineUserCloudEOSPlusPtr;
