Addition of TitleFile interface wrapper for EOSPlus

diff --git a/Engine/Plugins/Online/OnlineSubsystemEOS/Source/OnlineSubsystemEOSPlus/Private/OnlineSubsystemEOSPlus.cpp b/Engine/Plugins/Online/OnlineSubsystemEOS/Source/OnlineSubsystemEOSPlus/Private/OnlineSubsystemEOSPlus.cpp
--- a/Engine/Plugins/Online/OnlineSubsystemEOS/Source/OnlineSubsystemEOSPlus/Private/OnlineSubsystemEOSPlus.cpp
+++ b/Engine/Plugins/Online/OnlineSubsystemEOS/Source/OnlineSubsystemEOSPlus/Private/OnlineSubsystemEOSPlus.cpp
@@ -77,12 +77,14 @@ bool FOnlineSubsystemEOSPlus::Init()
 	LeaderboardsInterfacePtr = MakeShareable(new FOnlineLeaderboardsEOSPlus(this));
 	StoreInterfacePtr = MakeShareable(new FOnlineStoreEOSPlus(this));
 	ExternalUIInterfacePtr = MakeShareable(new FOnlineExternalUIEOSPlus(this));
+	TitleFileInterfacePtr = MakeShareable(new FOnlineTitleFileEOSPlus(this));
 	UserCloudInterfacePtr = MakeShareable(new FOnlineUserCloudEOSPlus(this));
 
 	UserInterfacePtr->Initialize();
 	LeaderboardsInterfacePtr->Initialize();
 	StoreInterfacePtr->Initialize();
 	ExternalUIInterfacePtr->Initialize();
+	TitleFileInterfacePtr->Initialize();
 	UserCloudInterfacePtr->Initialize();
 
 	return true;
@@ -106,6 +108,7 @@ void FOnlineSubsystemEOSPlus::PreUnload()
 	DESTRUCT_INTERFACE(LeaderboardsInterfacePtr);
 	DESTRUCT_INTERFACE(StoreInterfacePtr);
 	DESTRUCT_INTERFACE(ExternalUIInterfacePtr);
+	DESTRUCT_INTERFACE(TitleFileInterfacePtr);
 	DESTRUCT_INTERFACE(UserCloudInterfacePtr);
 
 #undef DESTRUCT_INTERFACE
@@ -181,7 +184,7 @@ IOnlineIdentityPtr FOnlineSubsystemEOSPlus::GetIdentityInterface() const
 
 IOnlineTitleFilePtr FOnlineSubsystemEOSPlus::GetTitleFileInterface() const
 {
-	return nullptr;
+	return TitleFileInterfacePtr;
 }
 
 IOnlineStoreV2Ptr FOnlineSubsystemEOSPlus::GetStoreV2Interface() const
@@ -266,4 +269,4 @@ bool FOnlineSubsystemEOSPlus::IsLocalPlayer(const FUniqueNetId& UniqueId) const
 	}
 
 	return false;
-}
\ No newline at end of file
+}
diff --git a/Engine/Plugins/Online/OnlineSubsystemEOS/Source/OnlineSubsystemEOSPlus/Private/OnlineSubsystemEOSPlus.h b/Engine/Plugins/Online/OnlineSubsystemEOS/Source/OnlineSubsystemEOSPlus/Private/OnlineSubsystemEOSPlus.h
--- a/Engine/Plugins/Online/OnlineSubsystemEOS/Source/OnlineSubsystemEOSPlus/Private/OnlineSubsystemEOSPlus.h
+++ b/Engine/Plugins/Online/OnlineSubsystemEOS/Source/OnlineSubsystemEOSPlus/Private/OnlineSubsystemEOSPlus.h
@@ -13,6 +13,7 @@
 #include "OnlineLeaderboardsEOSPlus.h"
 #include "OnlineStoreEOSPlus.h"
 #include "OnlineExternalUIEOSPlus.h"
+#include "OnlineTitleFileEOSPlus.h"
 #include "OnlineUserCloudEOSPlus.h"
 
 
@@ -85,6 +86,7 @@ PACKAGE_SCOPE:
 	FOnlineLeaderboardsEOSPlusPtr LeaderboardsInterfacePtr;
 	FOnlineStoreEOSPlusPtr StoreInterfacePtr;
 	FOnlineExternalUIEOSPlusPtr ExternalUIInterfacePtr;
+	FOnlineTitleFileEOSPlusPtr TitleFileInterfacePtr;
 	FOnlineUserCloudEOSPlusPtr UserCloudInterfacePtr;
 };
 
diff --git a/Engine/Plugins/Online/OnlineSubsystemEOS/Source/OnlineSubsystemEOSPlus/Private/OnlineTitleFileEOSPlus.cpp b/Engine/Plugins/Online/OnlineSubsystemEOS/Source/OnlineSubsystemEOSPlus/Private/OnlineTitleFileEOSPlus.cpp
new file mode 100644
--- /dev/null
+++ b/Engine/Plugins/Online/OnlineSubsystemEOS/Source/OnlineSubsystemEOSPlus/Private/OnlineTitleFileEOSPlus.cpp
@@ -0,0 +1,170 @@
+// Copyright Epic Games, Inc. All Rights Reserved.
+
+#include "OnlineTitleFileEOSPlus.h"
+#include "OnlineSubsystemEOSPlus.h"
+#include "OnlineError.h"
+
+FOnlineTitleFileEOSPlus::FOnlineTitleFileEOSPlus(FOnlineSubsystemEOSPlus* InSubsystem)
+	: EOSPlus(InSubsystem)
+{
+	BaseTitleFileInterface = EOSPlus->BaseOSS->GetTitleFileInterface();
+}
+
+FOnlineTitleFileEOSPlus::~FOnlineTitleFileEOSPlus()
+{
+	if (BaseTitleFileInterface.IsValid())
+	{
+		BaseTitleFileInterface->ClearOnEnumerateFilesCompleteDelegates(this);
+		BaseTitleFileInterface->ClearOnReadFileProgressDelegates(this);
+		BaseTitleFileInterface->ClearOnReadFileCompleteDelegates(this);
+		BaseTitleFileInterface->ClearOnTitleFileAnalyticsEventDelegates(this);
+	}
+}
+
+FUniqueNetIdEOSPlusPtr FOnlineTitleFileEOSPlus::GetNetIdPlus(const FString& SourceId) const
+{
+	return EOSPlus->UserInterfacePtr->GetNetIdPlus(SourceId);
+}
+
+void FOnlineTitleFileEOSPlus::Initialize()
+{
+	if (BaseTitleFileInterface.IsValid())
+	{
+		BaseTitleFileInterface->AddOnEnumerateFilesCompleteDelegate_Handle(FOnEnumerateFilesCompleteDelegate::CreateThreadSafeSP(this, &FOnlineTitleFileEOSPlus::OnEnumerateFilesComplete));
+		BaseTitleFileInterface->AddOnReadFileProgressDelegate_Handle(FOnReadFileProgressDelegate::CreateThreadSafeSP(this, &FOnlineTitleFileEOSPlus::OnReadFileProgress));
+		BaseTitleFileInterface->AddOnReadFileCompleteDelegate_Handle(FOnReadFileCompleteDelegate::CreateThreadSafeSP(this, &FOnlineTitleFileEOSPlus::OnReadFileComplete));
+		BaseTitleFileInterface->AddOnTitleFileAnalyticsEventDelegate_Handle(FOnTitleFileAnalyticsEventDelegate::CreateThreadSafeSP(this, &FOnlineTitleFileEOSPlus::OnTitleFileAnalyticsEvent));
+	}
+	else
+	{
+		UE_LOG_ONLINE(VeryVerbose, TEXT("[FOnlineTitleFileEOSPlus::Initialize] BaseTitleFileInterface delegates not bound. Base interface not valid"));
+	}
+}
+
+//~ Begin IOnlineTitleFile Interface
+
+bool FOnlineTitleFileEOSPlus::GetFileContents(const FString& FileName, TArray<uint8>& FileContents)
+{
+	bool bResult = false;
+
+	if (BaseTitleFileInterface.IsValid())
+	{
+		bResult = BaseTitleFileInterface->GetFileContents(FileName, FileContents);
+	}
+	else
+	{
+		UE_LOG_ONLINE(VeryVerbose, TEXT("[FOnlineTitleFileEOSPlus::GetFileContents] Unable to call method in base interface. Base interface not valid"));
+	}
+
+	return bResult;
+}
+
+bool FOnlineTitleFileEOSPlus::ClearFiles()
+{
+	bool bResult = false;
+
+	if (BaseTitleFileInterface.IsValid())
+	{
+		bResult = BaseTitleFileInterface->ClearFiles();
+	}
+	else
+	{
+		UE_LOG_ONLINE(VeryVerbose, TEXT("[FOnlineTitleFileEOSPlus::ClearFiles] Unable to call method in base interface. Base interface not valid"));
+	}
+
+	return bResult;
+}
+
+bool FOnlineTitleFileEOSPlus::ClearFile(const FString& FileName)
+{
+	bool bResult = false;
+
+	if (BaseTitleFileInterface.IsValid())
+	{
+		bResult = BaseTitleFileInterface->ClearFile(FileName);
+	}
+	else
+	{
+		UE_LOG_ONLINE(VeryVerbose, TEXT("[FOnlineTitleFileEOSPlus::ClearFile] Unable to call method in base interface. Base interface not valid"));
+	}
+
+	return bResult;
+}
+
+void FOnlineTitleFileEOSPlus::DeleteCachedFiles(bool bSkipEnumerated)
+{
+	if (BaseTitleFileInterface.IsValid())
+	{
+		BaseTitleFileInterface->DeleteCachedFiles(bSkipEnumerated);
+	}
+	else
+	{
+		UE_LOG_ONLINE(VeryVerbose, TEXT("[FOnlineTitleFileEOSPlus::DeleteCachedFiles] Unable to call method in base interface. Base interface not valid"));
+	}
+}
+
+bool FOnlineTitleFileEOSPlus::EnumerateFiles(const FPagedQuery& Page)
+{
+	bool bResult = false;
+
+	if (BaseTitleFileInterface.IsValid())
+	{
+		bResult = BaseTitleFileInterface->EnumerateFiles(Page);
+	}
+	else
+	{
+		UE_LOG_ONLINE(VeryVerbose, TEXT("[FOnlineTitleFileEOSPlus::EnumerateFiles] Unable to call method in base interface. Base interface not valid"));
+	}
+
+	return bResult;
+}
+
+void FOnlineTitleFileEOSPlus::GetFileList(TArray<FCloudFileHeader>& Files)
+{
+	if (BaseTitleFileInterface.IsValid())
+	{
+		BaseTitleFileInterface->GetFileList(Files);
+	}
+	else
+	{
+		UE_LOG_ONLINE(VeryVerbose, TEXT("[FOnlineTitleFileEOSPlus::GetFileList] Unable to call method in base interface. Base interface not valid"));
+	}
+}
+
+bool FOnlineTitleFileEOSPlus::ReadFile(const FString& FileName)
+{
+	bool bResult = false;
+
+	if (BaseTitleFileInterface.IsValid())
+	{
+		bResult = BaseTitleFileInterface->ReadFile(FileName);
+	}
+	else
+	{
+		UE_LOG_ONLINE(VeryVerbose, TEXT("[FOnlineTitleFileEOSPlus::ReadFile] Unable to call method in base interface. Base interface not valid"));
+	}
+
+	return bResult;
+}
+
+void FOnlineTitleFileEOSPlus::OnEnumerateFilesComplete(bool bWasSuccessful, const FString& ErrorStr)
+{
+	TriggerOnEnumerateFilesCompleteDelegates(bWasSuccessful, ErrorStr);
+}
+
+void FOnlineTitleFileEOSPlus::OnReadFileProgress(const FString& FileName, uint64 NumBytes)
+{
+	TriggerOnReadFileProgressDelegates(FileName, NumBytes);
+}
+
+void FOnlineTitleFileEOSPlus::OnReadFileComplete(bool bWasSuccessful, const FString& FileName)
+{
+	TriggerOnReadFileCompleteDelegates(bWasSuccessful, FileName);
+}
+
+void FOnlineTitleFileEOSPlus::OnTitleFileAnalyticsEvent(const FString& EventName, const TArray<FAnalyticsEventAttribute>& Attributes)
+{
+	TriggerOnTitleFileAnalyticsEventDelegates(EventName, Attributes);
+}
+
+//~ End IOnlineTitleFile Interface
\ No newline at end of file
diff --git a/Engine/Plugins/Online/OnlineSubsystemEOS/Source/OnlineSubsystemEOSPlus/Private/OnlineTitleFileEOSPlus.h b/Engine/Plugins/Online/OnlineSubsystemEOS/Source/OnlineSubsystemEOSPlus/Private/OnlineTitleFileEOSPlus.h
new file mode 100644
--- /dev/null
+++ b/Engine/Plugins/Online/OnlineSubsystemEOS/Source/OnlineSubsystemEOSPlus/Private/OnlineTitleFileEOSPlus.h
@@ -0,0 +1,54 @@
+// Copyright Epic Games, Inc. All Rights Reserved.
+
+#pragma once
+
+#include "CoreMinimal.h"
+#include "Online/CoreOnline.h"
+#include "OnlineSubsystem.h"
+#include "Interfaces/OnlineTitleFileInterface.h"
+#include "OnlineUserEOSPlus.h"
+
+class FOnlineSubsystemEOSPlus;
+
+/**
+ * Interface for encapsulating the platform title file interface
+ */
+class FOnlineTitleFileEOSPlus :
+	public IOnlineTitleFile,
+	public TSharedFromThis<FOnlineTitleFileEOSPlus, ESPMode::ThreadSafe>
+{
+public:
+	FOnlineTitleFileEOSPlus() = delete;
+	virtual ~FOnlineTitleFileEOSPlus();
+
+	void Initialize();
+
+	//~ Begin IOnlineTitleFile Interface
+	virtual bool GetFileContents(const FString& FileName, TArray<uint8>& FileContents) override;
+	virtual bool ClearFiles() override;
+	virtual bool ClearFile(const FString& FileName) override;
+	virtual void DeleteCachedFiles(bool bSkipEnumerated) override;
+	virtual bool EnumerateFiles(const FPagedQuery& Page = FPagedQuery()) override;
+	virtual void GetFileList(TArray<FCloudFileHeader>& Files) override;
+	virtual bool ReadFile(const FString& FileName) override;
+	//~ End IOnlineTitleFile Interface
+
+PACKAGE_SCOPE:
+	FOnlineTitleFileEOSPlus(FOnlineSubsystemEOSPlus* InSubsystem);
+
+	void OnEnumerateFilesComplete(bool bWasSuccessful, const FString& ErrorStr);
+	void OnReadFileProgress(const FString& FileName, uint64 NumBytes);
+	void OnReadFileComplete(bool bWasSuccessful, const FString& FileName);
+	void OnTitleFileAnalyticsEvent(const FString& EventName, const TArray<FAnalyticsEventAttribute>& Attributes);
+
+private:
+	FUniqueNetIdEOSPlusPtr GetNetIdPlus(const FString& SourceId) const;
+
+	/** Reference to the owning EOS plus subsystem */
+	FOnlineSubsystemEOSPlus* EOSPlus;
+	
+	// We don't support EOS mirroring yet
+	IOnlineTitleFilePtr BaseTitleFileInterface;
+};
+
+typedef TSharedPtr<FOnlineTitleFileEOSPlus, ESPMode::ThreadSafe> FOnlineTitleFileEOSPlusPtr;
