Addition of Voice interface wrapper for EOSPlus

diff --git a/Engine/Plugins/Online/OnlineSubsystemEOS/Source/OnlineSubsystemEOSPlus/Private/OnlineSubsystemEOSPlus.cpp b/Engine/Plugins/Online/OnlineSubsystemEOS/Source/OnlineSubsystemEOSPlus/Private/OnlineSubsystemEOSPlus.cpp
--- a/Engine/Plugins/Online/OnlineSubsystemEOS/Source/OnlineSubsystemEOSPlus/Private/OnlineSubsystemEOSPlus.cpp
+++ b/Engine/Plugins/Online/OnlineSubsystemEOS/Source/OnlineSubsystemEOSPlus/Private/OnlineSubsystemEOSPlus.cpp
@@ -77,6 +77,7 @@ bool FOnlineSubsystemEOSPlus::Init()
 	LeaderboardsInterfacePtr = MakeShareable(new FOnlineLeaderboardsEOSPlus(this));
 	StoreInterfacePtr = MakeShareable(new FOnlineStoreEOSPlus(this));
 	ExternalUIInterfacePtr = MakeShareable(new FOnlineExternalUIEOSPlus(this));
+	VoiceInterfacePtr = MakeShareable(new FOnlineVoiceEOSPlus(this));
 	TitleFileInterfacePtr = MakeShareable(new FOnlineTitleFileEOSPlus(this));
 	UserCloudInterfacePtr = MakeShareable(new FOnlineUserCloudEOSPlus(this));
 
@@ -84,6 +85,7 @@ bool FOnlineSubsystemEOSPlus::Init()
 	LeaderboardsInterfacePtr->Initialize();
 	StoreInterfacePtr->Initialize();
 	ExternalUIInterfacePtr->Initialize();
+	VoiceInterfacePtr->Initialize();
 	TitleFileInterfacePtr->Initialize();
 	UserCloudInterfacePtr->Initialize();
 
@@ -108,6 +110,7 @@ void FOnlineSubsystemEOSPlus::PreUnload()
 	DESTRUCT_INTERFACE(LeaderboardsInterfacePtr);
 	DESTRUCT_INTERFACE(StoreInterfacePtr);
 	DESTRUCT_INTERFACE(ExternalUIInterfacePtr);
+	DESTRUCT_INTERFACE(VoiceInterfacePtr);
 	DESTRUCT_INTERFACE(TitleFileInterfacePtr);
 	DESTRUCT_INTERFACE(UserCloudInterfacePtr);
 
@@ -164,7 +167,7 @@ IOnlineLeaderboardsPtr FOnlineSubsystemEOSPlus::GetLeaderboardsInterface() const
 
 IOnlineVoicePtr FOnlineSubsystemEOSPlus::GetVoiceInterface() const
 {
-	return nullptr;
+	return VoiceInterfacePtr;
 }
 
 IOnlineExternalUIPtr FOnlineSubsystemEOSPlus::GetExternalUIInterface() const
diff --git a/Engine/Plugins/Online/OnlineSubsystemEOS/Source/OnlineSubsystemEOSPlus/Private/OnlineSubsystemEOSPlus.h b/Engine/Plugins/Online/OnlineSubsystemEOS/Source/OnlineSubsystemEOSPlus/Private/OnlineSubsystemEOSPlus.h
--- a/Engine/Plugins/Online/OnlineSubsystemEOS/Source/OnlineSubsystemEOSPlus/Private/OnlineSubsystemEOSPlus.h
+++ b/Engine/Plugins/Online/OnlineSubsystemEOS/Source/OnlineSubsystemEOSPlus/Private/OnlineSubsystemEOSPlus.h
@@ -13,6 +13,7 @@
 #include "OnlineLeaderboardsEOSPlus.h"
 #include "OnlineStoreEOSPlus.h"
 #include "OnlineExternalUIEOSPlus.h"
+#include "OnlineVoiceEOSPlus.h"
 #include "OnlineTitleFileEOSPlus.h"
 #include "OnlineUserCloudEOSPlus.h"
 
@@ -86,6 +87,7 @@ PACKAGE_SCOPE:
 	FOnlineLeaderboardsEOSPlusPtr LeaderboardsInterfacePtr;
 	FOnlineStoreEOSPlusPtr StoreInterfacePtr;
 	FOnlineExternalUIEOSPlusPtr ExternalUIInterfacePtr;
+	FOnlineVoiceEOSPlusPtr VoiceInterfacePtr;
 	FOnlineTitleFileEOSPlusPtr TitleFileInterfacePtr;
 	FOnlineUserCloudEOSPlusPtr UserCloudInterfacePtr;
 };
diff --git a/Engine/Plugins/Online/OnlineSubsystemEOS/Source/OnlineSubsystemEOSPlus/Private/OnlineVoiceEOSPlus.cpp b/Engine/Plugins/Online/OnlineSubsystemEOS/Source/OnlineSubsystemEOSPlus/Private/OnlineVoiceEOSPlus.cpp
new file mode 100644
--- /dev/null
+++ b/Engine/Plugins/Online/OnlineSubsystemEOS/Source/OnlineSubsystemEOSPlus/Private/OnlineVoiceEOSPlus.cpp
@@ -0,0 +1,480 @@
+// Copyright Epic Games, Inc. All Rights Reserved.
+
+#include "OnlineVoiceEOSPlus.h"
+#include "OnlineSubsystemEOSPlus.h"
+#include "OnlineError.h"
+
+FOnlineVoiceEOSPlus::FOnlineVoiceEOSPlus(FOnlineSubsystemEOSPlus* InSubsystem)
+	: EOSPlus(InSubsystem)
+{
+	BaseVoiceInterface = EOSPlus->BaseOSS->GetVoiceInterface();
+}
+
+FOnlineVoiceEOSPlus::~FOnlineVoiceEOSPlus()
+{
+	if (BaseVoiceInterface.IsValid())
+	{
+		BaseVoiceInterface->ClearOnPlayerTalkingStateChangedDelegates(this);
+	}
+}
+
+FUniqueNetIdEOSPlusPtr FOnlineVoiceEOSPlus::GetNetIdPlus(const FString& SourceId) const
+{
+	return EOSPlus->UserInterfacePtr->GetNetIdPlus(SourceId);
+}
+
+void FOnlineVoiceEOSPlus::Initialize()
+{
+	if (BaseVoiceInterface.IsValid())
+	{
+		BaseVoiceInterface->AddOnPlayerTalkingStateChangedDelegate_Handle(FOnPlayerTalkingStateChangedDelegate::CreateThreadSafeSP(this, &FOnlineVoiceEOSPlus::OnPlayerTalkingStateChanged));
+	}
+	else
+	{
+		UE_LOG_ONLINE(VeryVerbose, TEXT("[FOnlineVoiceEOSPlus::Initialize] BaseVoiceInterface delegates not bound. Base interface not valid"));
+	}
+}
+
+//~ Begin IOnlineVoice Interface
+
+IVoiceEnginePtr FOnlineVoiceEOSPlus::CreateVoiceEngine()
+{
+	IVoiceEnginePtr Result = nullptr;
+	
+	if (BaseVoiceInterface.IsValid())
+	{
+		Result = BaseVoiceInterface->CreateVoiceEngine();
+	}
+	else
+	{
+		UE_LOG_ONLINE(VeryVerbose, TEXT("[FOnlineVoiceEOSPlus::CreateVoiceEngine] Unable to call method in base interface. Base interface not valid"));
+	}
+	
+	return Result;
+}
+
+void FOnlineVoiceEOSPlus::ProcessMuteChangeNotification()
+{
+	if (BaseVoiceInterface.IsValid())
+	{
+		BaseVoiceInterface->ProcessMuteChangeNotification();
+	}
+	else
+	{
+		UE_LOG_ONLINE(VeryVerbose, TEXT("[FOnlineVoiceEOSPlus::ProcessMuteChangeNotification] Unable to call method in base interface. Base interface not valid"));
+	}
+}
+
+void FOnlineVoiceEOSPlus::StartNetworkedVoice(uint8 LocalUserNum)
+{
+	if (BaseVoiceInterface.IsValid())
+	{
+		BaseVoiceInterface->StartNetworkedVoice(LocalUserNum);
+	}
+	else
+	{
+		UE_LOG_ONLINE(VeryVerbose, TEXT("[FOnlineVoiceEOSPlus::StartNetworkedVoice] Unable to call method in base interface. Base interface not valid"));
+	}
+}
+
+void FOnlineVoiceEOSPlus::StopNetworkedVoice(uint8 LocalUserNum)
+{
+	if (BaseVoiceInterface.IsValid())
+	{
+		BaseVoiceInterface->StopNetworkedVoice(LocalUserNum);
+	}
+	else
+	{
+		UE_LOG_ONLINE(VeryVerbose, TEXT("[FOnlineVoiceEOSPlus::StopNetworkedVoice] Unable to call method in base interface. Base interface not valid"));
+	}
+}
+
+bool FOnlineVoiceEOSPlus::RegisterLocalTalker(uint32 LocalUserNum)
+{
+	bool bResult = false;
+
+	if (BaseVoiceInterface.IsValid())
+	{
+		bResult = BaseVoiceInterface->RegisterLocalTalker(LocalUserNum);
+	}
+	else
+	{
+		UE_LOG_ONLINE(VeryVerbose, TEXT("[FOnlineVoiceEOSPlus::RegisterLocalTalker] Unable to call method in base interface. Base interface not valid"));
+	}
+
+	return bResult;
+}
+
+void FOnlineVoiceEOSPlus::RegisterLocalTalkers()
+{
+	if (BaseVoiceInterface.IsValid())
+	{
+		BaseVoiceInterface->RegisterLocalTalkers();
+	}
+	else
+	{
+		UE_LOG_ONLINE(VeryVerbose, TEXT("[FOnlineVoiceEOSPlus::RegisterLocalTalkers] Unable to call method in base interface. Base interface not valid"));
+	}
+}
+
+bool FOnlineVoiceEOSPlus::UnregisterLocalTalker(uint32 LocalUserNum)
+{
+	bool bResult = false;
+
+	if (BaseVoiceInterface.IsValid())
+	{
+		bResult = BaseVoiceInterface->UnregisterLocalTalker(LocalUserNum);
+	}
+	else
+	{
+		UE_LOG_ONLINE(VeryVerbose, TEXT("[FOnlineVoiceEOSPlus::UnregisterLocalTalker] Unable to call method in base interface. Base interface not valid"));
+	}
+
+	return bResult;
+}
+
+void FOnlineVoiceEOSPlus::UnregisterLocalTalkers()
+{
+	if (BaseVoiceInterface.IsValid())
+	{
+		BaseVoiceInterface->UnregisterLocalTalkers();
+	}
+	else
+	{
+		UE_LOG_ONLINE(VeryVerbose, TEXT("[FOnlineVoiceEOSPlus::UnregisterLocalTalkers] Unable to call method in base interface. Base interface not valid"));
+	}
+}
+
+bool FOnlineVoiceEOSPlus::RegisterRemoteTalker(const FUniqueNetId& UniqueId)
+{
+	bool bResult = false;
+
+	if (BaseVoiceInterface.IsValid())
+	{
+		FUniqueNetIdEOSPlusPtr NetIdPlus = GetNetIdPlus(UniqueId.ToString());
+		if (NetIdPlus.IsValid())
+		{
+			if (ensure(NetIdPlus->GetBaseNetId().IsValid()))
+			{
+				bResult = BaseVoiceInterface->RegisterRemoteTalker(*NetIdPlus->GetBaseNetId());
+			}
+			else
+			{
+				UE_LOG_ONLINE(Warning, TEXT("[FOnlineVoiceEOSPlus::RegisterRemoteTalker] Unable to call method in base interface. BaseNetId not valid for user (%s)."), *UniqueId.ToString());
+			}
+		}
+		else
+		{
+			UE_LOG_ONLINE(Warning, TEXT("[FOnlineVoiceEOSPlus::RegisterRemoteTalker] Unable to call method in base interface. Unknown user (%s)"), *UniqueId.ToString());
+		}
+	}
+	else
+	{
+		UE_LOG_ONLINE(VeryVerbose, TEXT("[FOnlineVoiceEOSPlus::RegisterRemoteTalker] Unable to call method in base interface. Base interface not valid"));
+	}
+
+	return bResult;
+}
+
+bool FOnlineVoiceEOSPlus::UnregisterRemoteTalker(const FUniqueNetId& UniqueId)
+{
+	bool bResult = false;
+
+	if (BaseVoiceInterface.IsValid())
+	{
+		FUniqueNetIdEOSPlusPtr NetIdPlus = GetNetIdPlus(UniqueId.ToString());
+		if (NetIdPlus.IsValid())
+		{
+			if (ensure(NetIdPlus->GetBaseNetId().IsValid()))
+			{
+				bResult = BaseVoiceInterface->UnregisterRemoteTalker(*NetIdPlus->GetBaseNetId());
+			}
+			else
+			{
+				UE_LOG_ONLINE(Warning, TEXT("[FOnlineVoiceEOSPlus::UnregisterRemoteTalker] Unable to call method in base interface. BaseNetId not valid for user (%s)."), *UniqueId.ToString());
+			}
+		}
+		else
+		{
+			UE_LOG_ONLINE(Warning, TEXT("[FOnlineVoiceEOSPlus::UnregisterRemoteTalker] Unable to call method in base interface. Unknown user (%s)"), *UniqueId.ToString());
+		}
+	}
+	else
+	{
+		UE_LOG_ONLINE(VeryVerbose, TEXT("[FOnlineVoiceEOSPlus::UnregisterRemoteTalker] Unable to call method in base interface. Base interface not valid"));
+	}
+
+	return bResult;
+}
+
+void FOnlineVoiceEOSPlus::RemoveAllRemoteTalkers()
+{
+	if (BaseVoiceInterface.IsValid())
+	{
+		BaseVoiceInterface->RemoveAllRemoteTalkers();
+	}
+	else
+	{
+		UE_LOG_ONLINE(VeryVerbose, TEXT("[FOnlineVoiceEOSPlus::RemoveAllRemoteTalkers] Unable to call method in base interface. Base interface not valid"));
+	}
+}
+
+bool FOnlineVoiceEOSPlus::IsHeadsetPresent(uint32 LocalUserNum)
+{
+	bool bResult = false;
+
+	if (BaseVoiceInterface.IsValid())
+	{
+		bResult = BaseVoiceInterface->IsHeadsetPresent(LocalUserNum);
+	}
+	else
+	{
+		UE_LOG_ONLINE(VeryVerbose, TEXT("[FOnlineVoiceEOSPlus::IsHeadsetPresent] Unable to call method in base interface. Base interface not valid"));
+	}
+
+	return bResult;
+}
+
+bool FOnlineVoiceEOSPlus::IsLocalPlayerTalking(uint32 LocalUserNum)
+{
+	bool bResult = false;
+
+	if (BaseVoiceInterface.IsValid())
+	{
+		bResult = BaseVoiceInterface->IsLocalPlayerTalking(LocalUserNum);
+	}
+	else
+	{
+		UE_LOG_ONLINE(VeryVerbose, TEXT("[FOnlineVoiceEOSPlus::IsLocalPlayerTalking] Unable to call method in base interface. Base interface not valid"));
+	}
+
+	return bResult;
+}
+
+bool FOnlineVoiceEOSPlus::IsRemotePlayerTalking(const FUniqueNetId& UniqueId)
+{
+	bool bResult = false;
+
+	if (BaseVoiceInterface.IsValid())
+	{
+		FUniqueNetIdEOSPlusPtr NetIdPlus = GetNetIdPlus(UniqueId.ToString());
+		if (NetIdPlus.IsValid())
+		{
+			if (ensure(NetIdPlus->GetBaseNetId().IsValid()))
+			{
+				bResult = BaseVoiceInterface->IsRemotePlayerTalking(*NetIdPlus->GetBaseNetId());
+			}
+			else
+			{
+				UE_LOG_ONLINE(Warning, TEXT("[FOnlineVoiceEOSPlus::IsRemotePlayerTalking] Unable to call method in base interface. BaseNetId not valid for user (%s)."), *UniqueId.ToString());
+			}
+		}
+		else
+		{
+			UE_LOG_ONLINE(Warning, TEXT("[FOnlineVoiceEOSPlus::IsRemotePlayerTalking] Unable to call method in base interface. Unknown user (%s)"), *UniqueId.ToString());
+		}
+	}
+	else
+	{
+		UE_LOG_ONLINE(VeryVerbose, TEXT("[FOnlineVoiceEOSPlus::IsRemotePlayerTalking] Unable to call method in base interface. Base interface not valid"));
+	}
+
+	return bResult;
+}
+
+bool FOnlineVoiceEOSPlus::IsMuted(uint32 LocalUserNum, const FUniqueNetId& UniqueId) const
+{
+	bool bResult = false;
+
+	if (BaseVoiceInterface.IsValid())
+	{
+		FUniqueNetIdEOSPlusPtr NetIdPlus = GetNetIdPlus(UniqueId.ToString());
+		if (NetIdPlus.IsValid())
+		{
+			if (ensure(NetIdPlus->GetBaseNetId().IsValid()))
+			{
+				bResult = BaseVoiceInterface->IsMuted(LocalUserNum, *NetIdPlus->GetBaseNetId());
+			}
+			else
+			{
+				UE_LOG_ONLINE(Warning, TEXT("[FOnlineVoiceEOSPlus::IsMuted] Unable to call method in base interface. BaseNetId not valid for user (%s)."), *UniqueId.ToString());
+			}
+		}
+		else
+		{
+			UE_LOG_ONLINE(Warning, TEXT("[FOnlineVoiceEOSPlus::IsMuted] Unable to call method in base interface. Unknown user (%s)"), *UniqueId.ToString());
+		}
+	}
+	else
+	{
+		UE_LOG_ONLINE(VeryVerbose, TEXT("[FOnlineVoiceEOSPlus::IsMuted] Unable to call method in base interface. Base interface not valid"));
+	}
+
+	return bResult;
+}
+
+bool FOnlineVoiceEOSPlus::MuteRemoteTalker(uint8 LocalUserNum, const FUniqueNetId& PlayerId, bool bIsSystemWide)
+{
+	bool bResult = false;
+
+	if (BaseVoiceInterface.IsValid())
+	{
+		FUniqueNetIdEOSPlusPtr NetIdPlus = GetNetIdPlus(PlayerId.ToString());
+		if (NetIdPlus.IsValid())
+		{
+			if (ensure(NetIdPlus->GetBaseNetId().IsValid()))
+			{
+				bResult = BaseVoiceInterface->MuteRemoteTalker(LocalUserNum, *NetIdPlus->GetBaseNetId(), bIsSystemWide);
+			}
+			else
+			{
+				UE_LOG_ONLINE(Warning, TEXT("[FOnlineVoiceEOSPlus::MuteRemoteTalker] Unable to call method in base interface. BaseNetId not valid for user (%s)."), *PlayerId.ToString());
+			}
+		}
+		else
+		{
+			UE_LOG_ONLINE(Warning, TEXT("[FOnlineVoiceEOSPlus::MuteRemoteTalker] Unable to call method in base interface. Unknown user (%s)"), *PlayerId.ToString());
+		}
+	}
+	else
+	{
+		UE_LOG_ONLINE(VeryVerbose, TEXT("[FOnlineVoiceEOSPlus::MuteRemoteTalker] Unable to call method in base interface. Base interface not valid"));
+	}
+
+	return bResult;
+}
+
+bool FOnlineVoiceEOSPlus::UnmuteRemoteTalker(uint8 LocalUserNum, const FUniqueNetId& PlayerId, bool bIsSystemWide)
+{
+	bool bResult = false;
+
+	if (BaseVoiceInterface.IsValid())
+	{
+		FUniqueNetIdEOSPlusPtr NetIdPlus = GetNetIdPlus(PlayerId.ToString());
+		if (NetIdPlus.IsValid())
+		{
+			if (ensure(NetIdPlus->GetBaseNetId().IsValid()))
+			{
+				bResult = BaseVoiceInterface->UnmuteRemoteTalker(LocalUserNum, *NetIdPlus->GetBaseNetId(), bIsSystemWide);
+			}
+			else
+			{
+				UE_LOG_ONLINE(Warning, TEXT("[FOnlineVoiceEOSPlus::UnmuteRemoteTalker] Unable to call method in base interface. BaseNetId not valid for user (%s)."), *PlayerId.ToString());
+			}
+		}
+		else
+		{
+			UE_LOG_ONLINE(Warning, TEXT("[FOnlineVoiceEOSPlus::UnmuteRemoteTalker] Unable to call method in base interface. Unknown user (%s)"), *PlayerId.ToString());
+		}
+	}
+	else
+	{
+		UE_LOG_ONLINE(VeryVerbose, TEXT("[FOnlineVoiceEOSPlus::UnmuteRemoteTalker] Unable to call method in base interface. Base interface not valid"));
+	}
+
+	return bResult;
+}
+
+TSharedPtr<FVoicePacket> FOnlineVoiceEOSPlus::SerializeRemotePacket(FArchive& Ar)
+{
+	TSharedPtr<FVoicePacket> Result = nullptr;
+	
+	if (BaseVoiceInterface.IsValid())
+	{
+		Result = BaseVoiceInterface->SerializeRemotePacket(Ar);
+	}
+	else
+	{
+		UE_LOG_ONLINE(VeryVerbose, TEXT("[FOnlineVoiceEOSPlus::SerializeRemotePacket] Unable to call method in base interface. Base interface not valid"));
+	}
+	
+	return Result;
+}
+
+TSharedPtr<FVoicePacket> FOnlineVoiceEOSPlus::GetLocalPacket(uint32 LocalUserNum)
+{
+	TSharedPtr<FVoicePacket> Result = nullptr;
+	
+	if (BaseVoiceInterface.IsValid())
+	{
+		Result = BaseVoiceInterface->GetLocalPacket(LocalUserNum);
+	}
+	else
+	{
+		UE_LOG_ONLINE(VeryVerbose, TEXT("[FOnlineVoiceEOSPlus::GetLocalPacket] Unable to call method in base interface. Base interface not valid"));
+	}
+	
+	return Result;
+}
+
+int32 FOnlineVoiceEOSPlus::GetNumLocalTalkers()
+{
+	int32 Result = 0;
+
+	if (BaseVoiceInterface.IsValid())
+	{
+		Result = BaseVoiceInterface->GetNumLocalTalkers();
+	}
+	else
+	{
+		UE_LOG_ONLINE(VeryVerbose, TEXT("[FOnlineVoiceEOSPlus::GetNumLocalTalkers] Unable to call method in base interface. Base interface not valid"));
+	}
+
+	return Result;
+}
+
+void FOnlineVoiceEOSPlus::ClearVoicePackets()
+{
+	if (BaseVoiceInterface.IsValid())
+	{
+		BaseVoiceInterface->ClearVoicePackets();
+	}
+	else
+	{
+		UE_LOG_ONLINE(VeryVerbose, TEXT("[FOnlineVoiceEOSPlus::ClearVoicePackets] Unable to call method in base interface. Base interface not valid"));
+	}
+}
+
+void FOnlineVoiceEOSPlus::Tick(float DeltaTime)
+{
+	if (BaseVoiceInterface.IsValid())
+	{
+		BaseVoiceInterface->Tick(DeltaTime);
+	}
+	else
+	{
+		UE_LOG_ONLINE(VeryVerbose, TEXT("[FOnlineVoiceEOSPlus::Tick] Unable to call method in base interface. Base interface not valid"));
+	}
+}
+
+FString FOnlineVoiceEOSPlus::GetVoiceDebugState() const
+{
+	FString Result;
+
+	if (BaseVoiceInterface.IsValid())
+	{
+		Result = BaseVoiceInterface->GetVoiceDebugState();
+	}
+	else
+	{
+		UE_LOG_ONLINE(VeryVerbose, TEXT("[FOnlineVoiceEOSPlus::GetVoiceDebugState] Unable to call method in base interface. Base interface not valid"));
+	}
+
+	return Result;
+}
+	
+void FOnlineVoiceEOSPlus::OnPlayerTalkingStateChanged(FUniqueNetIdRef TalkerId, bool bIsTalking)
+{
+	FUniqueNetIdEOSPlusPtr NetIdPlus = GetNetIdPlus(TalkerId->ToString());
+	if (NetIdPlus.IsValid())
+	{
+		// We don't want to trigger it anyways if the NetIdPlus is not valid like in other interfaces, because this one takes a FUniqueNetIdRef as parameter
+		TriggerOnPlayerTalkingStateChangedDelegates(NetIdPlus.ToSharedRef(), bIsTalking);
+	}
+	else
+	{
+		UE_LOG_ONLINE(Warning, TEXT("[FOnlineVoiceEOSPlus::OnPlayerTalkingStateChanged] Unknown user (%s)"), *TalkerId->ToString());
+	}	
+}
+
+//~ End IOnlineVoice Interface
\ No newline at end of file
diff --git a/Engine/Plugins/Online/OnlineSubsystemEOS/Source/OnlineSubsystemEOSPlus/Private/OnlineVoiceEOSPlus.h b/Engine/Plugins/Online/OnlineSubsystemEOS/Source/OnlineSubsystemEOSPlus/Private/OnlineVoiceEOSPlus.h
new file mode 100644
--- /dev/null
+++ b/Engine/Plugins/Online/OnlineSubsystemEOS/Source/OnlineSubsystemEOSPlus/Private/OnlineVoiceEOSPlus.h
@@ -0,0 +1,68 @@
+// Copyright Epic Games, Inc. All Rights Reserved.
+
+#pragma once
+
+#include "CoreMinimal.h"
+#include "Online/CoreOnline.h"
+#include "OnlineSubsystem.h"
+#include "Interfaces/VoiceInterface.h"
+#include "OnlineUserEOSPlus.h"
+
+class FOnlineSubsystemEOSPlus;
+
+/**
+ * Interface for encapsulating the platform voice interface
+ */
+class FOnlineVoiceEOSPlus :
+	public IOnlineVoice,
+	public TSharedFromThis<FOnlineVoiceEOSPlus, ESPMode::ThreadSafe>
+{
+public:
+	FOnlineVoiceEOSPlus() = delete;
+	virtual ~FOnlineVoiceEOSPlus();
+
+	void Initialize();
+
+	//~ Begin IOnlineVoice Interface
+	virtual void StartNetworkedVoice(uint8 LocalUserNum) override;
+	virtual void StopNetworkedVoice(uint8 LocalUserNum) override;
+	virtual bool RegisterLocalTalker(uint32 LocalUserNum) override;
+	virtual void RegisterLocalTalkers() override;
+	virtual bool UnregisterLocalTalker(uint32 LocalUserNum) override;
+	virtual void UnregisterLocalTalkers() override;
+	virtual bool RegisterRemoteTalker(const FUniqueNetId& UniqueId) override;
+	virtual bool UnregisterRemoteTalker(const FUniqueNetId& UniqueId) override;
+	virtual void RemoveAllRemoteTalkers() override;
+	virtual bool IsHeadsetPresent(uint32 LocalUserNum) override;
+	virtual bool IsLocalPlayerTalking(uint32 LocalUserNum) override;
+	virtual bool IsRemotePlayerTalking(const FUniqueNetId& UniqueId) override;
+	virtual bool IsMuted(uint32 LocalUserNum, const FUniqueNetId& UniqueId) const override;
+	virtual bool MuteRemoteTalker(uint8 LocalUserNum, const FUniqueNetId& PlayerId, bool bIsSystemWide) override;
+	virtual bool UnmuteRemoteTalker(uint8 LocalUserNum, const FUniqueNetId& PlayerId, bool bIsSystemWide) override;
+	virtual TSharedPtr<FVoicePacket> SerializeRemotePacket(FArchive& Ar) override;
+	virtual TSharedPtr<FVoicePacket> GetLocalPacket(uint32 LocalUserNum) override;
+	virtual int32 GetNumLocalTalkers() override;
+	virtual void ClearVoicePackets() override;
+	virtual void Tick(float DeltaTime) override;
+	virtual FString GetVoiceDebugState() const override;
+	//~ End IOnlineVoice Interface
+
+PACKAGE_SCOPE:
+	FOnlineVoiceEOSPlus(FOnlineSubsystemEOSPlus* InSubsystem);
+
+	virtual IVoiceEnginePtr CreateVoiceEngine() override;
+	virtual void ProcessMuteChangeNotification() override;
+	
+	void OnPlayerTalkingStateChanged(FUniqueNetIdRef TalkerId, bool bIsTalking);
+
+private:
+	FUniqueNetIdEOSPlusPtr GetNetIdPlus(const FString& SourceId) const;
+
+	/** Reference to the owning EOS plus subsystem */
+	FOnlineSubsystemEOSPlus* EOSPlus;
+	
+	// We don't support EOS mirroring yet
+	IOnlineVoicePtr BaseVoiceInterface;
+};
+
+typedef TSharedPtr<FOnlineVoiceEOSPlus, ESPMode::ThreadSafe> FOnlineVoiceEOSPlusPtr;
