Refactor and improvements based on review feedback for CL-16796668

diff --git a/Engine/Plugins/Online/OnlineSubsystemEOS/Source/OnlineSubsystemEOS/Private/UserManagerEOS.cpp b/Engine/Plugins/Online/OnlineSubsystemEOS/Source/OnlineSubsystemEOS/Private/UserManagerEOS.cpp
--- a/Engine/Plugins/Online/OnlineSubsystemEOS/Source/OnlineSubsystemEOS/Private/UserManagerEOS.cpp
+++ b/Engine/Plugins/Online/OnlineSubsystemEOS/Source/OnlineSubsystemEOS/Private/UserManagerEOS.cpp
@@ -1648,7 +1648,7 @@ bool FUserManagerEOS::ReadFriendsList(int32 LocalUserNum, const FString& ListNam
 			Options.LocalUserId = UserNumToAccountIdMap[LocalUserNum];
 			int32 FriendCount = EOS_Friends_GetFriendsCount(EOSSubsystem->FriendsHandle, &Options);
 
-			LocalUserNumToFriendsListMap[LocalUserNum]->Empty();
+			LocalUserNumToFriendsListMap[LocalUserNum]->Empty(FriendCount);
 
 			// Process each friend returned
 			for (int32 Index = 0; Index < FriendCount; Index++)
@@ -1683,17 +1683,17 @@ void FUserManagerEOS::ProcessReadFriendsListComplete(int32 LocalUserNum, bool bW
 	if (!IsFriendQueryUserInfoOngoing(LocalUserNum))
 	{
 		// If not, we'll just trigger the delegates for all cached calls
-		if (CachedReadUserListInfoForLocalUserMap.Contains(LocalUserNum))
-		{			
-			for (const ReadUserListInfo& CachedInfo : CachedReadUserListInfoForLocalUserMap[LocalUserNum])
+		TArray<ReadUserListInfo> CachedInfoList;
+		if (CachedReadUserListInfoForLocalUserMap.RemoveAndCopyValue(LocalUserNum, CachedInfoList))
+		{
+			for (const ReadUserListInfo& CachedInfo : CachedInfoList)
 			{
 				CachedInfo.ExecuteDelegateIfBound(bWasSuccessful, ErrorStr);
-				TriggerOnFriendsChangeDelegates(LocalUserNum);
 			}
-
-			CachedReadUserListInfoForLocalUserMap.Remove(LocalUserNum);
 		}
 
+		TriggerOnFriendsChangeDelegates(LocalUserNum);
+
 		IsFriendQueryUserInfoOngoingForLocalUserMap.Remove(LocalUserNum);
 	}
 }
@@ -2404,12 +2404,7 @@ void FUserManagerEOS::ReadUserInfo(int32 LocalUserNum, EOS_EpicAccountId EpicAcc
 	EOS_UserInfo_QueryUserInfo(EOSSubsystem->UserInfoHandle, &Options, CallbackObj, CallbackObj->GetCallbackPtr());
 
 	// We mark this player as pending for processing
-	if (!IsFriendQueryUserInfoOngoingForLocalUserMap.Contains(LocalUserNum))
-	{
-		IsFriendQueryUserInfoOngoingForLocalUserMap.Emplace(LocalUserNum);
-}
-
-	IsFriendQueryUserInfoOngoingForLocalUserMap[LocalUserNum].Add(EpicAccountId);
+	IsFriendQueryUserInfoOngoingForLocalUserMap.FindOrAdd(LocalUserNum).Add(EpicAccountId);
 }
 
 bool FUserManagerEOS::GetAllUserInfo(int32 LocalUserNum, TArray<TSharedRef<FOnlineUser>>& OutUsers)
diff --git a/Engine/Plugins/Online/OnlineSubsystemEOS/Source/OnlineSubsystemEOS/Private/UserManagerEOS.h b/Engine/Plugins/Online/OnlineSubsystemEOS/Source/OnlineSubsystemEOS/Private/UserManagerEOS.h
--- a/Engine/Plugins/Online/OnlineSubsystemEOS/Source/OnlineSubsystemEOS/Private/UserManagerEOS.h
+++ b/Engine/Plugins/Online/OnlineSubsystemEOS/Source/OnlineSubsystemEOS/Private/UserManagerEOS.h
@@ -141,10 +141,10 @@ public:
 		ListEntries.Remove(InListEntry);
 	}
 
-	void Empty()
+	void Empty(int32 Slack = 0)
 	{
-		ListEntries.Empty();
-		NetIdStringToListEntryMap.Empty();
+		ListEntries.Empty(Slack);
+		NetIdStringToListEntryMap.Empty(Slack);
 	}
 
 	void UpdateNetIdStr(const FString& PrevNetId, const FString& NewNetId)
@@ -462,7 +462,8 @@ private:
 		{
 		}
 
-		void ExecuteDelegateIfBound(bool bWasSuccessful, const FString& ErrorStr) const {
+		void ExecuteDelegateIfBound(bool bWasSuccessful, const FString& ErrorStr) const
+		{
 			Delegate.ExecuteIfBound(LocalUserNum, bWasSuccessful, ListName, ErrorStr);
 		};
 	};
