-Refactor of UEOSSettings class so that its CDO is loaded successfully in all platforms.
-Removed forced logout on login status changed in EOSPlus.

diff --git a/Engine/Plugins/Online/OnlineSubsystemEOS/Source/OnlineSubsystemEOS/Private/EOSSettings.cpp b/Engine/Plugins/Online/OnlineSubsystemEOS/Source/OnlineSubsystemEOS/Private/EOSSettings.cpp
--- a/Engine/Plugins/Online/OnlineSubsystemEOS/Source/OnlineSubsystemEOS/Private/EOSSettings.cpp
+++ b/Engine/Plugins/Online/OnlineSubsystemEOS/Source/OnlineSubsystemEOS/Private/EOSSettings.cpp
@@ -52,7 +52,7 @@ inline bool ContainsWhitespace(const FString& Source)
 	return false;
 }
 
-/*FEOSArtifactSettings FArtifactSettings::ToNative() const
+FEOSArtifactSettings FArtifactSettings::ToNative() const
 {
 	FEOSArtifactSettings Native;
 
@@ -65,7 +65,7 @@ inline bool ContainsWhitespace(const FString& Source)
 	Native.SandboxId = SandboxId;
 
 	return Native;
-}*/
+}
 
 inline FString StripQuotes(const FString& Source)
 {
@@ -121,20 +121,20 @@ void FEOSArtifactSettings::ParseRawArrayEntry(const FString& RawLine)
 	}
 }
 
-const FEOSSettings& UEOSSettings::GetSettings()
+FEOSSettings UEOSSettings::GetSettings()
 {
-	/*if (UObjectInitialized())
+	if (UObjectInitialized())
 	{
 		return UEOSSettings::AutoGetSettings();
-	}*/
+	}
 
 	return UEOSSettings::ManualGetSettings();
 }
 
-/*FEOSSettings UEOSSettings::AutoGetSettings()
+FEOSSettings UEOSSettings::AutoGetSettings()
 {
 	return GetDefault<UEOSSettings>()->ToNative();
-}*/
+}
 
 const FEOSSettings& UEOSSettings::ManualGetSettings()
 {
@@ -164,7 +164,7 @@ const FEOSSettings& UEOSSettings::ManualGetSettings()
 	return *CachedSettings;
 }
 
-/*FEOSSettings UEOSSettings::ToNative() const
+FEOSSettings UEOSSettings::ToNative() const
 {
 	FEOSSettings Native;
 
@@ -185,14 +185,14 @@ const FEOSSettings& UEOSSettings::ManualGetSettings()
 	Native.TitleStorageTags = TitleStorageTags;
 
 	return Native;
-}*/
+}
 
 bool UEOSSettings::GetSettingsForArtifact(const FString& ArtifactName, FEOSArtifactSettings& OutSettings)
 {
-	/*if (UObjectInitialized())
+	if (UObjectInitialized())
 	{
 		return UEOSSettings::AutoGetSettingsForArtifact(ArtifactName, OutSettings);
-	}*/
+	}
 	return UEOSSettings::ManualGetSettingsForArtifact(ArtifactName, OutSettings);
 }
 
@@ -252,7 +252,7 @@ bool UEOSSettings::ManualGetSettingsForArtifact(const FString& ArtifactName, FEO
 	return false;
 }
 
-/*bool UEOSSettings::AutoGetSettingsForArtifact(const FString& ArtifactName, FEOSArtifactSettings& OutSettings)
+bool UEOSSettings::AutoGetSettingsForArtifact(const FString& ArtifactName, FEOSArtifactSettings& OutSettings)
 {
 	const UEOSSettings* This = GetDefault<UEOSSettings>();
 	FString ArtifactNameOverride;
@@ -280,7 +280,7 @@ bool UEOSSettings::ManualGetSettingsForArtifact(const FString& ArtifactName, FEO
 	}
 	UE_LOG_ONLINE(Error, TEXT("UEOSSettings::AutoGetSettingsForArtifact() failed due to missing config object specified. Check your project settings"));
 	return false;
-}*/
+}
 
 #if WITH_EDITOR
 void UEOSSettings::PostEditChangeProperty(FPropertyChangedEvent& PropertyChangedEvent)
diff --git a/Engine/Plugins/Online/OnlineSubsystemEOS/Source/OnlineSubsystemEOS/Public/EOSSettings.h b/Engine/Plugins/Online/OnlineSubsystemEOS/Source/OnlineSubsystemEOS/Public/EOSSettings.h
--- a/Engine/Plugins/Online/OnlineSubsystemEOS/Source/OnlineSubsystemEOS/Public/EOSSettings.h
+++ b/Engine/Plugins/Online/OnlineSubsystemEOS/Source/OnlineSubsystemEOS/Public/EOSSettings.h
@@ -3,7 +3,7 @@
 #pragma once
 
 #include "UObject/ObjectMacros.h"
-#include "UObject/Object.h"
+#include "Engine/RuntimeOptionsBase.h"
 #include "Engine/DataAsset.h"
 #include "EOSSettings.generated.h"
 
@@ -61,7 +61,7 @@ public:
 	UPROPERTY(Config, EditAnywhere, BlueprintReadOnly, Category="EOS Artifact Settings")
 	FString EncryptionKey;
 
-	//FEOSArtifactSettings ToNative() const; // See note #1
+	FEOSArtifactSettings ToNative() const;
 };
 
 /** Native version of the UObject based config data */
@@ -86,7 +86,7 @@ struct FEOSSettings
 
 UCLASS(Config=Engine, DefaultConfig)
 class ONLINESUBSYSTEMEOS_API UEOSSettings :
-	public UObject
+	public URuntimeOptionsBase
 {
 	GENERATED_BODY()
 
@@ -158,21 +158,16 @@ public:
 	/** Find the Settings for an artifact by name */
 	static bool GetSettingsForArtifact(const FString& ArtifactName, FEOSArtifactSettings& OutSettings);
 
-	static const FEOSSettings& GetSettings();
-	//FEOSSettings ToNative() const; // See note #1
+	static FEOSSettings GetSettings();
+	FEOSSettings ToNative() const;
 
 private:
 #if WITH_EDITOR
 	virtual void PostEditChangeProperty(FPropertyChangedEvent& PropertyChangedEvent) override;
 #endif
-	//static bool AutoGetSettingsForArtifact(const FString& ArtifactName, FEOSArtifactSettings& OutSettings); // See note #1
+	static bool AutoGetSettingsForArtifact(const FString& ArtifactName, FEOSArtifactSettings& OutSettings);
 	static bool ManualGetSettingsForArtifact(const FString& ArtifactName, FEOSArtifactSettings& OutSettings);
 
-	//static FEOSSettings AutoGetSettings(); // See note #1
+	static FEOSSettings AutoGetSettings();
 	static const FEOSSettings& ManualGetSettings();
-};
-
-/* Notes
-* 
-*  #1: CDO Usage was unreliable in some platforms, so it will be deactivated for now
-*/
\ No newline at end of file
+};
\ No newline at end of file
diff --git a/Engine/Plugins/Online/OnlineSubsystemEOS/Source/OnlineSubsystemEOSPlus/Private/OnlineUserEOSPlus.cpp b/Engine/Plugins/Online/OnlineSubsystemEOS/Source/OnlineSubsystemEOSPlus/Private/OnlineUserEOSPlus.cpp
--- a/Engine/Plugins/Online/OnlineSubsystemEOS/Source/OnlineSubsystemEOSPlus/Private/OnlineUserEOSPlus.cpp
+++ b/Engine/Plugins/Online/OnlineSubsystemEOS/Source/OnlineSubsystemEOSPlus/Private/OnlineUserEOSPlus.cpp
@@ -515,16 +515,7 @@ void FOnlineUserEOSPlus::OnEOSLoginChanged(int32 LocalUserNum)
 
 void FOnlineUserEOSPlus::OnLoginStatusChanged(int32 LocalUserNum, ELoginStatus::Type OldStatus, ELoginStatus::Type NewStatus, const FUniqueNetId& NewId)
 {
-	if (NewStatus == ELoginStatus::UsingLocalProfile)
-	{
-		const FEOSSettings& EOSSettings = UEOSSettings::GetSettings();
-
-		if (!EOSSettings.bUseEAS && !EOSSettings.bUseEOSConnect)
-		{
-			EOSIdentityInterface->Logout(LocalUserNum);
-		}
-		TriggerOnLoginStatusChangedDelegates(LocalUserNum, OldStatus, NewStatus, NewId);
-	}
+	TriggerOnLoginStatusChangedDelegates(LocalUserNum, OldStatus, NewStatus, NewId);
 }
 
 void FOnlineUserEOSPlus::OnControllerPairingChanged(int32 LocalUserNum, FControllerPairingChangedUserInfo PreviousUser, FControllerPairingChangedUserInfo NewUser)
