Added missing conversion

diff --git a/Engine/Plugins/Online/OnlineSubsystemEOS/Source/OnlineSubsystemEOS/Private/UserManagerEOS.cpp b/Engine/Plugins/Online/OnlineSubsystemEOS/Source/OnlineSubsystemEOS/Private/UserManagerEOS.cpp
--- a/Engine/Plugins/Online/OnlineSubsystemEOS/Source/OnlineSubsystemEOS/Private/UserManagerEOS.cpp
+++ b/Engine/Plugins/Online/OnlineSubsystemEOS/Source/OnlineSubsystemEOS/Private/UserManagerEOS.cpp
@@ -2292,7 +2292,7 @@ void FUserManagerEOS::UpdatePresence(EOS_EpicAccountId AccountId)
 		const FString Platform(UTF8_TO_TCHAR(PresenceInfo->Platform));
 		// Convert the presence data to our format
 		PresenceRef->Status.State = ToEOnlinePresenceState(PresenceInfo->Status);
-		PresenceRef->Status.StatusStr = PresenceInfo->RichText;
+		PresenceRef->Status.StatusStr = UTF8_TO_TCHAR(PresenceInfo->RichText);
 		PresenceRef->bIsOnline = PresenceRef->Status.State == EOnlinePresenceState::Online;
 		PresenceRef->bIsPlaying = !ProductId.IsEmpty();
 		PresenceRef->bIsPlayingThisGame = ProductId == EOSSubsystem->ProductId && ProdVersion == EOSSubsystem->EOSSDKManager->GetProductVersion();
