Added missing conversion to avoid player name character errors

diff --git a/Engine/Plugins/Online/OnlineSubsystemEOS/Source/OnlineSubsystemEOS/Private/UserManagerEOS.cpp b/Engine/Plugins/Online/OnlineSubsystemEOS/Source/OnlineSubsystemEOS/Private/UserManagerEOS.cpp
--- a/Engine/Plugins/Online/OnlineSubsystemEOS/Source/OnlineSubsystemEOS/Private/UserManagerEOS.cpp
+++ b/Engine/Plugins/Online/OnlineSubsystemEOS/Source/OnlineSubsystemEOS/Private/UserManagerEOS.cpp
@@ -871,9 +871,9 @@ void FUserManagerEOS::UpdateUserInfo(IAttributeAccessInterfaceRef AttributeAcces
 	EOS_EResult CopyResult = EOS_UserInfo_CopyUserInfo(EOSSubsystem->UserInfoHandle, &Options, &UserInfo);
 	if (CopyResult == EOS_EResult::EOS_Success)
 	{
-		AttributeAccessRef->SetInternalAttribute(USER_ATTR_DISPLAY_NAME, UserInfo->DisplayName);
-		AttributeAccessRef->SetInternalAttribute(USER_ATTR_COUNTRY, UserInfo->Country);
-		AttributeAccessRef->SetInternalAttribute(USER_ATTR_LANG, UserInfo->PreferredLanguage);
+		AttributeAccessRef->SetInternalAttribute(USER_ATTR_DISPLAY_NAME, UTF8_TO_TCHAR(UserInfo->DisplayName));
+		AttributeAccessRef->SetInternalAttribute(USER_ATTR_COUNTRY, UTF8_TO_TCHAR(UserInfo->Country));
+		AttributeAccessRef->SetInternalAttribute(USER_ATTR_LANG, UTF8_TO_TCHAR(UserInfo->PreferredLanguage));
 		EOS_UserInfo_Release(UserInfo);
 	}
 }
