Fixed EOS OSS checking for being launched by the Epic Launcher when in PIE

diff --git a/Engine/Plugins/Online/OnlineSubsystemEOS/Source/OnlineSubsystemEOS/Private/OnlineSubsystemEOS.cpp b/Engine/Plugins/Online/OnlineSubsystemEOS/Source/OnlineSubsystemEOS/Private/OnlineSubsystemEOS.cpp
--- a/Engine/Plugins/Online/OnlineSubsystemEOS/Source/OnlineSubsystemEOS/Private/OnlineSubsystemEOS.cpp
+++ b/Engine/Plugins/Online/OnlineSubsystemEOS/Source/OnlineSubsystemEOS/Private/OnlineSubsystemEOS.cpp
@@ -171,7 +171,7 @@ bool FOnlineSubsystemEOS::Init()
 	// Check for being launched by EGS
 	bWasLaunchedByEGS = FParse::Param(FCommandLine::Get(), TEXT("EpicPortal"));
 	FEOSSettings EOSSettings = UEOSSettings::GetSettings();
-	if (!IsRunningDedicatedServer() && !bWasLaunchedByEGS && EOSSettings.bShouldEnforceBeingLaunchedByEGS)
+	if (!IsRunningDedicatedServer() && IsRunningGame() && !bWasLaunchedByEGS && EOSSettings.bShouldEnforceBeingLaunchedByEGS)
 	{
 		FString ArtifactName;
 		FParse::Value(FCommandLine::Get(), TEXT("EpicApp="), ArtifactName);
