Removed unnecessary ensure

diff --git a/Engine/Plugins/Online/OnlineSubsystemEOS/Source/OnlineSubsystemEOS/Private/OnlineSubsystemEOS.cpp b/Engine/Plugins/Online/OnlineSubsystemEOS/Source/OnlineSubsystemEOS/Private/OnlineSubsystemEOS.cpp
--- a/Engine/Plugins/Online/OnlineSubsystemEOS/Source/OnlineSubsystemEOS/Private/OnlineSubsystemEOS.cpp
+++ b/Engine/Plugins/Online/OnlineSubsystemEOS/Source/OnlineSubsystemEOS/Private/OnlineSubsystemEOS.cpp
@@ -318,24 +318,15 @@ bool FOnlineSubsystemEOS::Shutdown()
 	}
 #endif
 
-#define DESTRUCT_INTERFACE(Interface) \
-	if (Interface.IsValid()) \
-	{ \
-		ensure(Interface.IsUnique()); \
-		Interface = nullptr; \
-	}
-
-	// Destruct the interfaces
-	DESTRUCT_INTERFACE(UserManager);
-	DESTRUCT_INTERFACE(SessionInterfacePtr);
-	DESTRUCT_INTERFACE(StatsInterfacePtr);
-	DESTRUCT_INTERFACE(LeaderboardsInterfacePtr);
-	DESTRUCT_INTERFACE(AchievementsInterfacePtr);
-	DESTRUCT_INTERFACE(StoreInterfacePtr);
-	DESTRUCT_INTERFACE(TitleFileInterfacePtr);
-	DESTRUCT_INTERFACE(UserCloudInterfacePtr);
-
-#undef DESTRUCT_INTERFACE
+	// Release our ref to the interfaces. May still exist since they can be aggregated
+	UserManager = nullptr;
+	SessionInterfacePtr = nullptr;
+	StatsInterfacePtr = nullptr;
+	LeaderboardsInterfacePtr = nullptr;
+	AchievementsInterfacePtr = nullptr;
+	StoreInterfacePtr = nullptr;
+	TitleFileInterfacePtr = nullptr;
+	UserCloudInterfacePtr = nullptr;
 
 	return true;
 }
