Fixed a memory aliasing bug

diff --git a/Engine/Plugins/Online/OnlineSubsystemEOS/Source/OnlineSubsystemEOS/Private/UserManagerEOS.cpp b/Engine/Plugins/Online/OnlineSubsystemEOS/Source/OnlineSubsystemEOS/Private/UserManagerEOS.cpp
--- a/Engine/Plugins/Online/OnlineSubsystemEOS/Source/OnlineSubsystemEOS/Private/UserManagerEOS.cpp
+++ b/Engine/Plugins/Online/OnlineSubsystemEOS/Source/OnlineSubsystemEOS/Private/UserManagerEOS.cpp
@@ -1533,16 +1533,16 @@ void FUserManagerEOS::AddRemotePlayer(const FString& NetId, EOS_EpicAccountId Ep
 void FUserManagerEOS::UpdateRemotePlayerProductUserId(EOS_EpicAccountId AccountId, EOS_ProductUserId UserId)
 {
 	// See if the net ids have changed for this user and bail if they are the same
-	FString NewNetIdStr = MakeNetIdStringFromIds(AccountId, UserId);
-	const FString& PrevNetIdStr = AccountIdToStringMap[AccountId];
+	const FString NewNetIdStr = MakeNetIdStringFromIds(AccountId, UserId);
+	const FString PrevNetIdStr = AccountIdToStringMap[AccountId];
 	if (PrevNetIdStr == NewNetIdStr)
 	{
 		// No change, so skip any work
 		return;
 	}
 
-	FString AccountIdStr = MakeStringFromEpicAccountId(AccountId);
-	FString UserIdStr = MakeStringFromProductUserId(UserId);
+	const FString AccountIdStr = MakeStringFromEpicAccountId(AccountId);
+	const FString UserIdStr = MakeStringFromProductUserId(UserId);
 
 	// Get the unique net id and rebuild the string for it
 	IAttributeAccessInterfaceRef AttrAccess = NetIdStringToAttributeAccessMap[PrevNetIdStr];
