Fixed some missing cases when transitioning away from bUsesPresence to specify lobbies

diff --git a/Engine/Plugins/Online/OnlineSubsystem/Source/Public/OnlineSessionSettings.h b/Engine/Plugins/Online/OnlineSubsystem/Source/Public/OnlineSessionSettings.h
--- a/Engine/Plugins/Online/OnlineSubsystem/Source/Public/OnlineSessionSettings.h
+++ b/Engine/Plugins/Online/OnlineSubsystem/Source/Public/OnlineSessionSettings.h
@@ -637,6 +637,8 @@ public:
 #define SEARCH_XBOX_LIVE_SESSION_TEMPLATE_NAME FName(TEXT("LIVESESSIONTEMPLATE"))
 /** Selection method used to determine which match to join when multiple are returned (valid only on Switch) */
 #define SEARCH_SWITCH_SELECTION_METHOD FName(TEXT("SWITCHSELECTIONMETHOD"))
+/** Whether to use lobbies vs sessions */
+#define SEARCH_LOBBIES FName(TEXT("LOBBYSEARCH"))
 
 // User attributes for searching (FSessionMatchmakingUser::Attributes)
 /** Team a user is searching for */
diff --git a/Engine/Plugins/Online/OnlineSubsystemEOS/Source/OnlineSubsystemEOS/Private/OnlineSessionEOS.cpp b/Engine/Plugins/Online/OnlineSubsystemEOS/Source/OnlineSubsystemEOS/Private/OnlineSessionEOS.cpp
--- a/Engine/Plugins/Online/OnlineSubsystemEOS/Source/OnlineSubsystemEOS/Private/OnlineSessionEOS.cpp
+++ b/Engine/Plugins/Online/OnlineSubsystemEOS/Source/OnlineSubsystemEOS/Private/OnlineSessionEOS.cpp
@@ -496,7 +496,7 @@ FNamedOnlineSession* FOnlineSessionEOS::GetNamedSessionFromLobbyId(const FUnique
 		if (Session.SessionInfo.IsValid())
 		{
 			FOnlineSessionInfoEOS* SessionInfo = (FOnlineSessionInfoEOS*)Session.SessionInfo.Get();
-			if (!Session.SessionSettings.bIsLANMatch && Session.SessionSettings.bUsesPresence && *SessionInfo->SessionId == LobbyId)
+			if (!Session.SessionSettings.bIsLANMatch && Session.SessionSettings.bUseLobbiesIfAvailable && *SessionInfo->SessionId == LobbyId)
 			{
 				Result = &Sessions[SearchIndex];
 				break;
@@ -1293,7 +1293,7 @@ bool FOnlineSessionEOS::UpdateSession(FName SessionName, FOnlineSessionSettings&
 
 		if (!Session->SessionSettings.bIsLANMatch)
 		{
-			if (Session->SessionSettings.bUsesPresence)
+			if (Session->SessionSettings.bUseLobbiesIfAvailable)
 			{
 				Result = UpdateLobbySession(Session);
 			}
@@ -1378,7 +1378,7 @@ bool FOnlineSessionEOS::EndSession(FName SessionName)
 		{
 			if (!Session->SessionSettings.bIsLANMatch)
 			{
-				if (Session->SessionSettings.bUsesPresence)
+				if (Session->SessionSettings.bUseLobbiesIfAvailable)
 				{
 					Result = EndLobbySession(Session);
 				}
@@ -1485,7 +1485,7 @@ bool FOnlineSessionEOS::DestroySession(FName SessionName, const FOnDestroySessio
 			{
 				if (Session->SessionState == EOnlineSessionState::InProgress)
 				{
-					if (Session->SessionSettings.bUsesPresence)
+					if (Session->SessionSettings.bUseLobbiesIfAvailable)
 					{
 						Result = EndLobbySession(Session);
 					}
@@ -1495,7 +1495,7 @@ bool FOnlineSessionEOS::DestroySession(FName SessionName, const FOnDestroySessio
 					}
 				}
 
-				if (Session->SessionSettings.bUsesPresence)
+				if (Session->SessionSettings.bUseLobbiesIfAvailable)
 				{
 					Result = DestroyLobbySession(Session, CompletionDelegate);
 				}
@@ -1672,8 +1672,8 @@ bool FOnlineSessionEOS::FindSessions(int32 SearchingPlayerNum, const TSharedRef<
 		// Check if its a LAN query
 		if (!SearchSettings->bIsLanQuery)
 		{
-			bool bUsesPresence = false;
-			if (SearchSettings->QuerySettings.Get(SEARCH_PRESENCE, bUsesPresence) && bUsesPresence)
+			bool bUssLobbiesIfAvailable = false;
+			if (SearchSettings->QuerySettings.Get(SEARCH_LOBBIES, bUssLobbiesIfAvailable) && bUssLobbiesIfAvailable)
 			{
 				Return = FindLobbySession(SearchingPlayerNum, SearchSettings);
 			}
@@ -2053,7 +2053,7 @@ bool FOnlineSessionEOS::JoinSession(int32 PlayerNum, FName SessionName, const FO
 				FOnlineSessionInfoEOS* NewSessionInfo = new FOnlineSessionInfoEOS(*SearchSessionInfo);
 				Session->SessionInfo = MakeShareable(NewSessionInfo);
 
-				if (DesiredSession.Session.SessionSettings.bUsesPresence)
+				if (DesiredSession.Session.SessionSettings.bUseLobbiesIfAvailable)
 				{
 					Return = JoinLobbySession(PlayerNum, Session, &DesiredSession.Session);
 				}
@@ -2258,7 +2258,7 @@ bool FOnlineSessionEOS::SendSessionInvite(FName SessionName, EOS_ProductUserId S
 	FNamedOnlineSession* Session = GetNamedSession(SessionName);
 	if (Session != nullptr)
 	{
-		if (Session->SessionSettings.bUsesPresence)
+		if (Session->SessionSettings.bUseLobbiesIfAvailable)
 		{
 			bResult = SendLobbyInvite(SessionName, SenderId, ReceiverId);
 		}
@@ -3362,7 +3362,7 @@ uint32 FOnlineSessionEOS::DestroyLobbySession(FNamedOnlineSession* Session, cons
 		Session->SessionState = EOnlineSessionState::Destroying;
 
 		FOnlineSessionInfoEOS* SessionInfo = (FOnlineSessionInfoEOS*)(Session->SessionInfo.Get());
-		check(Session->SessionSettings.bUsesPresence); // We check if it's a lobby session
+		check(Session->SessionSettings.bUseLobbiesIfAvailable); // We check if it's a lobby session
 
 		FName SessionName = Session->SessionName;
 
@@ -3631,13 +3631,14 @@ void FOnlineSessionEOS::AddLobbySearchResult(EOS_HLobbyDetails LobbyDetailsHandl
 
 void FOnlineSessionEOS::CopyLobbyData(EOS_HLobbyDetails LobbyDetailsHandle, EOS_LobbyDetails_Info* LobbyDetailsInfo, FOnlineSession& OutSession)
 {
-	OutSession.SessionSettings.bUsesPresence = true;
+	OutSession.SessionSettings.bUseLobbiesIfAvailable = true;
 	OutSession.SessionSettings.bIsLANMatch = false;
 
 	switch (LobbyDetailsInfo->PermissionLevel)
 	{
 	case EOS_ELobbyPermissionLevel::EOS_LPL_PUBLICADVERTISED:
 	case EOS_ELobbyPermissionLevel::EOS_LPL_JOINVIAPRESENCE:
+		OutSession.SessionSettings.bUsesPresence = true;
 		OutSession.SessionSettings.bAllowJoinViaPresence = true;
 
 		OutSession.SessionSettings.NumPublicConnections = LobbyDetailsInfo->MaxMembers;
@@ -3645,6 +3646,7 @@ void FOnlineSessionEOS::CopyLobbyData(EOS_HLobbyDetails LobbyDetailsHandle, EOS_
 
 		break;
 	case EOS_ELobbyPermissionLevel::EOS_LPL_INVITEONLY:
+		OutSession.SessionSettings.bUsesPresence = false;
 		OutSession.SessionSettings.bAllowJoinViaPresence = false;
 
 		OutSession.SessionSettings.NumPrivateConnections = LobbyDetailsInfo->MaxMembers;
@@ -3654,7 +3656,7 @@ void FOnlineSessionEOS::CopyLobbyData(EOS_HLobbyDetails LobbyDetailsHandle, EOS_
 	}
 
 	OutSession.SessionSettings.bAllowInvites = (bool)LobbyDetailsInfo->bAllowInvites;
-	
+
 	CopyLobbyAttributes(LobbyDetailsHandle, OutSession);
 }
 
