- Add UObject settings object to EOSShared to control the Log level, and potentially other common init attributes.
- Add platform handle type to auto-release the platform when it goes out of scope.
- Fix uninitialized variables in UEOSSettings/FEOSSettings

diff --git a/Engine/Plugins/Online/EOSShared/Source/EOSShared/EOSShared.Build.cs b/Engine/Plugins/Online/EOSShared/Source/EOSShared/EOSShared.Build.cs
--- a/Engine/Plugins/Online/EOSShared/Source/EOSShared/EOSShared.Build.cs
+++ b/Engine/Plugins/Online/EOSShared/Source/EOSShared/EOSShared.Build.cs
@@ -10,17 +10,13 @@ public class EOSShared : ModuleRules
 	{
 		Type = ModuleType.CPlusPlus;
 
-		PublicDependencyModuleNames.AddRange(
-			new string[]
-			{
-			}
-		);
-
 		PrivateDependencyModuleNames.AddRange(
 			new string[]
 			{
 				"Core",
-				"EOSSDK"
+				"CoreUObject",
+				"EOSSDK",
+				"Settings"
 			}
 		);
 	}
diff --git a/Engine/Plugins/Online/EOSShared/Source/EOSShared/Private/EOSSDKManager.cpp b/Engine/Plugins/Online/EOSShared/Source/EOSShared/Private/EOSSDKManager.cpp
--- a/Engine/Plugins/Online/EOSShared/Source/EOSShared/Private/EOSSDKManager.cpp
+++ b/Engine/Plugins/Online/EOSShared/Source/EOSShared/Private/EOSSDKManager.cpp
@@ -8,19 +8,20 @@
 #include "HAL/LowLevelMemTracker.h"
 #include "Misc/App.h"
 #include "Misc/ConfigCacheIni.h"
+#include "Misc/CoreMisc.h"
 #include "Misc/NetworkVersion.h"
 #include "Stats/Stats.h"
 
 #include "CoreGlobals.h"
-#include "EOSSharedUtils.h"
+
+#include "EOSShared.h"
+#include "EOSSharedSettings.h"
 
 #include "eos_init.h"
 #include "eos_logging.h"
 #include "eos_sdk.h"
 #include "eos_version.h"
 
-DEFINE_LOG_CATEGORY_STATIC(LogEOSSDK, Log, All);
-
 namespace
 {
 	static void* EOS_MEMORY_CALL EosMalloc(size_t Bytes, size_t Alignment)
@@ -126,8 +127,6 @@ EOS_EResult FEOSSDKManager::Initialize()
 	}
 	else
 	{
-		UpdateConfiguration();
-
 		UE_LOG(LogEOSSDK, Log, TEXT("Initializing EOSSDK Version:%s"), UTF8_TO_TCHAR(EOS_GetVersion()));
 
 		const FTCHARToUTF8 ProductName(*GetProductName());
@@ -159,36 +158,10 @@ EOS_EResult FEOSSDKManager::Initialize()
 				UE_LOG(LogEOSSDK, Warning, TEXT("EOS_Logging_SetCallback failed error:%s"), *LexToString(EosResult));
 			}
 
-			EOS_ELogLevel LogLevel = EOS_ELogLevel::EOS_LOG_Info;
-			if (Config.LogLevel == TEXT("Off"))
-			{
-				LogLevel = EOS_ELogLevel::EOS_LOG_Off;
-			}
-			else if (Config.LogLevel == TEXT("Fatal"))
-			{
-				LogLevel = EOS_ELogLevel::EOS_LOG_Fatal;
-			}
-			else if (Config.LogLevel == TEXT("Error"))
-			{
-				LogLevel = EOS_ELogLevel::EOS_LOG_Error;
-			}
-			else if (Config.LogLevel == TEXT("Warning"))
-			{
-				LogLevel = EOS_ELogLevel::EOS_LOG_Warning;
-			}
-			else if (Config.LogLevel == TEXT("Info"))
-			{
-				LogLevel = EOS_ELogLevel::EOS_LOG_Info;
-			}
-			else if (Config.LogLevel == TEXT("Verbose"))
-			{
-				LogLevel = EOS_ELogLevel::EOS_LOG_Verbose;
-			}
-			else if (Config.LogLevel == TEXT("VeryVerbose"))
-			{
-				LogLevel = EOS_ELogLevel::EOS_LOG_VeryVerbose;
-			}
-			EosResult = EOS_Logging_SetLogLevel(EOS_ELogCategory::EOS_LC_ALL_CATEGORIES, LogLevel);
+			const FEOSSharedSettings& EOSSharedSettings = UEOSSharedSettings::GetSettings();
+
+			const EOS_ELogLevel EosLogLevel = ConvertLogLevel(EOSSharedSettings.LogLevel);
+			EosResult = EOS_Logging_SetLogLevel(EOS_ELogCategory::EOS_LC_ALL_CATEGORIES, EosLogLevel);
 			if (EosResult != EOS_EResult::EOS_Success)
 			{
 				UE_LOG(LogEOSSDK, Warning, TEXT("EOS_Logging_SetLogLevel failed error:%s"), *LexToString(EosResult));
@@ -203,16 +176,17 @@ EOS_EResult FEOSSDKManager::Initialize()
 	}
 }
 
-EOS_HPlatform FEOSSDKManager::CreatePlatform(const EOS_Platform_Options& PlatformOptions)
+IEOSPlatformHandlePtr FEOSSDKManager::CreatePlatform(const EOS_Platform_Options& PlatformOptions)
 {
-	EOS_HPlatform EosPlatformHandle = nullptr;
+	IEOSPlatformHandlePtr SharedPlatform;
 
 	if (IsInitialized())
 	{
-		EosPlatformHandle = EOS_Platform_Create(&PlatformOptions);
-		if (EosPlatformHandle)
+		const EOS_HPlatform PlatformHandle = EOS_Platform_Create(&PlatformOptions);
+		if (PlatformHandle)
 		{
-			EosPlatformHandles.Emplace(EosPlatformHandle);
+			PlatformHandles.Emplace(PlatformHandle);
+			SharedPlatform = MakeShared<FEOSPlatformHandle, ESPMode::ThreadSafe>(*this, PlatformHandle);
 
 			if (!TickerHandle.IsValid())
 			{
@@ -229,31 +203,16 @@ EOS_HPlatform FEOSSDKManager::CreatePlatform(const EOS_Platform_Options& Platfor
 		UE_LOG(LogEOSSDK, Warning, TEXT("FEOSSDKManager::CreatePlatform failed. SDK not initialized"));
 	}
 
-	return EosPlatformHandle;
-}
-
-void FEOSSDKManager::ReleasePlatform(EOS_HPlatform EosPlatformHandle)
-{
-	if (EosPlatformHandle)
-	{
-		EosPlatformHandles.Remove(EosPlatformHandle);
-		EOS_Platform_Release(EosPlatformHandle);
-
-		if (EosPlatformHandles.Num() == 0 && TickerHandle.IsValid())
-		{
-			FTicker::GetCoreTicker().RemoveTicker(TickerHandle);
-			TickerHandle.Reset();
-		}
-	}
+	return SharedPlatform;
 }
 
 bool FEOSSDKManager::Tick(float)
 {
 	//LLM_SCOPE(ELLMTag::EOSSDK); // TODO
-	for (EOS_HPlatform EosPlatformHandle : EosPlatformHandles)
+	for (EOS_HPlatform PlatformHandle : PlatformHandles)
 	{
 		QUICK_SCOPE_CYCLE_COUNTER(FEOSSDKManager_Tick);
-		EOS_Platform_Tick(EosPlatformHandle);
+		EOS_Platform_Tick(PlatformHandle);
 	}
 
 	return true;
@@ -270,16 +229,43 @@ FString FEOSSDKManager::GetProductVersion() const
 	return ProductVersion;
 }
 
+void FEOSSDKManager::ReleasePlatform(EOS_HPlatform PlatformHandle)
+{
+	// Only release the platform if it was actually present in PlatformHandles, as Shutdown may have already released it.
+	if (PlatformHandles.Remove(PlatformHandle))
+	{
+		EOS_Platform_Release(PlatformHandle);
+
+		if (ensure(TickerHandle.IsValid()) &&
+			PlatformHandles.Num() == 0)
+		{
+			FTicker::GetCoreTicker().RemoveTicker(TickerHandle);
+			TickerHandle.Reset();
+		}
+	}
+	else
+	{
+		UE_LOG(LogEOSSDK, Warning, TEXT("FEOSSDKManager::ReleasePlatform PlatformHandle does not exist."));
+	}
+}
+
 void FEOSSDKManager::Shutdown()
 {
 	if (IsInitialized())
 	{
-		if (EosPlatformHandles.Num() > 0)
+		if (PlatformHandles.Num() > 0)
 		{
-			UE_LOG(LogEOSSDK, Warning, TEXT("FEOSSDKManager::Shutdown Releasing %d remaining platforms"), EosPlatformHandles.Num());
-			for (EOS_HPlatform EosPlatformHandle : EosPlatformHandles)
+			UE_LOG(LogEOSSDK, Warning, TEXT("FEOSSDKManager::Shutdown Releasing %d remaining platforms"), PlatformHandles.Num());
+			for (EOS_HPlatform PlatformHandle : PlatformHandles)
 			{
-				ReleasePlatform(EosPlatformHandle);
+				EOS_Platform_Release(PlatformHandle);
+			}
+			PlatformHandles.Empty();
+
+			if (ensure(TickerHandle.IsValid()))
+			{
+				FTicker::GetCoreTicker().RemoveTicker(TickerHandle);
+				TickerHandle.Reset();
 			}
 		}
 
@@ -295,12 +281,15 @@ EOS_EResult FEOSSDKManager::EOSInitialize(EOS_InitializeOptions& Options)
 	return EOS_Initialize(&Options);
 }
 
-void FEOSSDKManager::UpdateConfiguration()
+FEOSPlatformHandle::~FEOSPlatformHandle()
 {
-	const TCHAR* SectionName = TEXT("EOSSDK");
-	const FString& ConfigFile = GEngineIni;
+	Manager.ReleasePlatform(PlatformHandle);
+}
 
-	GConfig->GetString(SectionName, TEXT("LogLevel"), Config.LogLevel, ConfigFile);
+void FEOSPlatformHandle::Tick()
+{
+	QUICK_SCOPE_CYCLE_COUNTER(FEOSPlatformHandle_Tick);
+	EOS_Platform_Tick(PlatformHandle);
 }
 
 #endif // WITH_EOS_SDK
\ No newline at end of file
diff --git a/Engine/Plugins/Online/EOSShared/Source/EOSShared/Private/EOSSDKManager.h b/Engine/Plugins/Online/EOSShared/Source/EOSShared/Private/EOSSDKManager.h
--- a/Engine/Plugins/Online/EOSShared/Source/EOSShared/Private/EOSSDKManager.h
+++ b/Engine/Plugins/Online/EOSShared/Source/EOSShared/Private/EOSSDKManager.h
@@ -14,6 +14,8 @@
 
 #include "IEOSSDKManager.h"
 
+struct FEOSPlatformHandle;
+
 class FEOSSDKManager : public IEOSSDKManager
 {
 public:
@@ -24,10 +26,7 @@ public:
 	virtual EOS_EResult Initialize() override;
 	virtual bool IsInitialized() const override { return bInitialized; }
 
-	virtual EOS_HPlatform CreatePlatform(const EOS_Platform_Options& PlatformOptions) override;
-	virtual void ReleasePlatform(EOS_HPlatform EosPlatformHandle) override;
-
-	virtual bool Tick(float) override;
+	virtual IEOSPlatformHandlePtr CreatePlatform(const EOS_Platform_Options& PlatformOptions) override;
 
 	virtual FString GetProductName() const override;
 	virtual FString GetProductVersion() const override;
@@ -38,23 +37,36 @@ public:
 protected:
 	virtual EOS_EResult EOSInitialize(EOS_InitializeOptions& Options);
 
-	void UpdateConfiguration();
+private:
+	friend struct FEOSPlatformHandle;
+
+	void ReleasePlatform(EOS_HPlatform PlatformHandle);
+	bool Tick(float);
 
 #if EOSSDK_RUNTIME_LOAD_REQUIRED
 	void* SDKHandle = nullptr;
 #endif
 
-	struct
-	{
-		FString LogLevel;
-	} Config;
-
 	/** Are we currently initialized */
 	bool bInitialized = false;
-	/** Set of EOS platforms created with CreatePlatform */
-	TSet<EOS_HPlatform> EosPlatformHandles;
+	/** Created platforms */
+	TArray<EOS_HPlatform> PlatformHandles;
 	/** Handle to ticker delegate for Tick(), valid whenever EosPlatformHandles is non-empty. */
 	FDelegateHandle TickerHandle;
 };
 
+struct FEOSPlatformHandle : public IEOSPlatformHandle
+{
+	FEOSPlatformHandle(FEOSSDKManager& InManager, EOS_HPlatform InPlatformHandle)
+		: IEOSPlatformHandle(InPlatformHandle), Manager(InManager)
+	{
+	}
+
+	virtual ~FEOSPlatformHandle();
+	virtual void Tick() override;
+
+	/* Reference to the EOSSDK manager */
+	FEOSSDKManager& Manager;
+};
+
 #endif // WITH_EOS_SDK
\ No newline at end of file
diff --git a/Engine/Plugins/Online/EOSShared/Source/EOSShared/Private/EOSShared.cpp b/Engine/Plugins/Online/EOSShared/Source/EOSShared/Private/EOSShared.cpp
new file mode 100644
--- /dev/null
+++ b/Engine/Plugins/Online/EOSShared/Source/EOSShared/Private/EOSShared.cpp
@@ -0,0 +1,5 @@
+// Copyright Epic Games, Inc. All Rights Reserved.
+
+#include "EOSShared.h"
+
+DEFINE_LOG_CATEGORY(LogEOSSDK);
\ No newline at end of file
diff --git a/Engine/Plugins/Online/EOSShared/Source/EOSShared/Private/EOSSharedModule.cpp b/Engine/Plugins/Online/EOSShared/Source/EOSShared/Private/EOSSharedModule.cpp
--- a/Engine/Plugins/Online/EOSShared/Source/EOSShared/Private/EOSSharedModule.cpp
+++ b/Engine/Plugins/Online/EOSShared/Source/EOSShared/Private/EOSSharedModule.cpp
@@ -3,10 +3,17 @@
 #include "EOSSharedModule.h"
 
 #include "Features/IModularFeatures.h"
+#include "Misc/CoreDelegates.h"
 #include "Modules/ModuleManager.h"
 
+#if WITH_EDITOR
+#include "ISettingsModule.h"
+#endif // WITH_EDITOR
+
 #include COMPILED_PLATFORM_HEADER(EOSSDKManager.h)
 
+#define LOCTEXT_NAMESPACE "EOS"
+
 IMPLEMENT_MODULE(FEOSSharedModule, EOSShared);
 
 void FEOSSharedModule::StartupModule()
@@ -16,11 +23,42 @@ void FEOSSharedModule::StartupModule()
 	check(SDKManager);
 
 	IModularFeatures::Get().RegisterModularFeature(TEXT("EOSSDKManager"), SDKManager.Get());
-#endif
+#endif // WITH_EOS_SDK
+
+#if WITH_EDITOR
+	FCoreDelegates::OnPostEngineInit.AddRaw(this, &FEOSSharedModule::OnPostEngineInit);
+	FCoreDelegates::OnPreExit.AddRaw(this, &FEOSSharedModule::OnPreExit);
+#endif // WITH_EDITOR
+}
+
+#if WITH_EDITOR
+void FEOSSharedModule::OnPostEngineInit()
+{
+	if (ISettingsModule* SettingsModule = FModuleManager::GetModulePtr<ISettingsModule>("Settings"))
+	{
+		SettingsModule->RegisterSettings("Project", "Plugins", "Epic Online Services SDK",
+			LOCTEXT("EOSSDKSettingsName", "Epic Online Services SDK"),
+			LOCTEXT("EOSSDKSettingsDescription", "Configure the Epic Online Services SDK"),
+			GetMutableDefault<UEOSSharedSettings>());
+	}
 }
 
+void FEOSSharedModule::OnPreExit()
+{
+	if (ISettingsModule* SettingsModule = FModuleManager::GetModulePtr<ISettingsModule>("Settings"))
+	{
+		SettingsModule->UnregisterSettings("Project", "Plugins", "Epic Online Services SDK");
+	}
+}
+#endif // WITH_EDITOR
+
 void FEOSSharedModule::ShutdownModule()
 {
+#if WITH_EDITOR
+	FCoreDelegates::OnPostEngineInit.RemoveAll(this);
+	FCoreDelegates::OnPreExit.RemoveAll(this);
+#endif // WITH_EDITOR
+
 #if WITH_EOS_SDK
 	if(SDKManager.IsValid())
 	{
@@ -28,5 +66,7 @@ void FEOSSharedModule::ShutdownModule()
 		SDKManager->Shutdown();
 		SDKManager.Reset();
 	}
-#endif
+#endif // WITH_EOS_SDK
 }
+
+#undef LOCTEXT_NAMESPACE
\ No newline at end of file
diff --git a/Engine/Plugins/Online/EOSShared/Source/EOSShared/Private/EOSSharedModule.h b/Engine/Plugins/Online/EOSShared/Source/EOSShared/Private/EOSSharedModule.h
--- a/Engine/Plugins/Online/EOSShared/Source/EOSShared/Private/EOSSharedModule.h
+++ b/Engine/Plugins/Online/EOSShared/Source/EOSShared/Private/EOSSharedModule.h
@@ -18,6 +18,9 @@ private:
 	virtual void ShutdownModule() override;
 	// ~End IModuleInterface
 
+	void OnPostEngineInit();
+	void OnPreExit();
+
 #if WITH_EOS_SDK
 	TUniquePtr<FEOSSDKManager> SDKManager;
 #endif
diff --git a/Engine/Plugins/Online/EOSShared/Source/EOSShared/Private/EOSSharedSettings.cpp b/Engine/Plugins/Online/EOSShared/Source/EOSShared/Private/EOSSharedSettings.cpp
new file mode 100644
--- /dev/null
+++ b/Engine/Plugins/Online/EOSShared/Source/EOSShared/Private/EOSSharedSettings.cpp
@@ -0,0 +1,126 @@
+// Copyright Epic Games, Inc. All Rights Reserved.
+
+#include "EOSSharedSettings.h"
+#include "EOSShared.h"
+
+#include "Algo/Transform.h"
+#include "Misc/CommandLine.h"
+#include "Misc/ConfigCacheIni.h"
+
+#if WITH_EDITOR
+#include "Misc/MessageDialog.h"
+#endif
+
+#define LOCTEXT_NAMESPACE "EOS"
+#define INI_SECTION TEXT("/Script/EOSShared.EOSSharedSettings")
+
+FEOSSharedSettings UEOSSharedSettings::GetSettings()
+{
+	if (UObjectInitialized())
+	{
+		return UEOSSharedSettings::AutoGetSettings();
+	}
+	return UEOSSharedSettings::ManualGetSettings();
+}
+
+FEOSSharedSettings UEOSSharedSettings::AutoGetSettings()
+{
+	return GetDefault<UEOSSharedSettings>()->ToNative();
+}
+
+FEOSSharedSettings UEOSSharedSettings::ManualGetSettings()
+{
+	FEOSSharedSettings Native;
+
+	FString LogLevelString;
+	if (GConfig->GetString(INI_SECTION, TEXT("LogLevel"), LogLevelString, GEngineIni))
+	{
+		LexFromString(Native.LogLevel, *LogLevelString);
+	}
+
+	return Native;
+}
+
+FEOSSharedSettings UEOSSharedSettings::ToNative() const
+{
+	FEOSSharedSettings Native;
+
+	Native.LogLevel = LogLevel;
+
+	return Native;
+}
+
+#if WITH_EDITOR
+void UEOSSharedSettings::PostEditChangeProperty(FPropertyChangedEvent& PropertyChangedEvent)
+{
+	if (PropertyChangedEvent.Property)
+	{
+		if (PropertyChangedEvent.Property->GetFName() == FName(TEXT("LogLevel")))
+		{
+			IEOSSDKManager* SDKManager = IEOSSDKManager::Get();
+			if (SDKManager && SDKManager->IsInitialized())
+			{
+				const EOS_ELogLevel EosLogLevel = ConvertLogLevel(LogLevel);
+				const EOS_EResult EosResult = EOS_Logging_SetLogLevel(EOS_ELogCategory::EOS_LC_ALL_CATEGORIES, EosLogLevel);
+				if (EosResult != EOS_EResult::EOS_Success)
+				{
+					UE_LOG(LogEOSSDK, Warning, TEXT("EOS_Logging_SetLogLevel failed error:%s"), *LexToString(EosResult));
+				}
+			}
+		}
+	}
+
+	Super::PostEditChangeProperty(PropertyChangedEvent);
+}
+#endif
+
+void LexFromString(EEosLogLevel& Value, const TCHAR* String)
+{
+	if (FCString::Stricmp(String, TEXT("Off")) == 0)
+	{
+		Value = EEosLogLevel::Off;
+	}
+	else if (FCString::Stricmp(String, TEXT("Fatal")) == 0)
+	{
+		Value = EEosLogLevel::Fatal;
+	}
+	else if (FCString::Stricmp(String, TEXT("Error")) == 0)
+	{
+		Value = EEosLogLevel::Error;
+	}
+	else if (FCString::Stricmp(String, TEXT("Warning")) == 0)
+	{
+		Value = EEosLogLevel::Warning;
+	}
+	else if (FCString::Stricmp(String, TEXT("Verbose")) == 0)
+	{
+		Value = EEosLogLevel::Verbose;
+	}
+	else if (FCString::Stricmp(String, TEXT("VeryVerbose")) == 0)
+	{
+		Value = EEosLogLevel::VeryVerbose;
+	}
+	else
+	{
+		Value = EEosLogLevel::Info;
+	}
+}
+
+EOS_ELogLevel ConvertLogLevel(const EEosLogLevel LogLevel)
+{
+	EOS_ELogLevel Result = EOS_ELogLevel::EOS_LOG_Info;
+	switch (LogLevel)
+	{
+	case EEosLogLevel::Off:			Result = EOS_ELogLevel::EOS_LOG_Off; break;
+	case EEosLogLevel::Fatal:		Result = EOS_ELogLevel::EOS_LOG_Fatal; break;
+	case EEosLogLevel::Error:		Result = EOS_ELogLevel::EOS_LOG_Error; break;
+	case EEosLogLevel::Warning:		Result = EOS_ELogLevel::EOS_LOG_Warning; break;
+	case EEosLogLevel::Info:		Result = EOS_ELogLevel::EOS_LOG_Info; break;
+	case EEosLogLevel::Verbose:		Result = EOS_ELogLevel::EOS_LOG_Verbose; break;
+	case EEosLogLevel::VeryVerbose:	Result = EOS_ELogLevel::EOS_LOG_VeryVerbose; break;
+	}
+	return Result;
+}
+
+#undef INI_SECTION
+#undef LOCTEXT_NAMESPACE
diff --git a/Engine/Plugins/Online/EOSShared/Source/EOSShared/Public/EOSShared.h b/Engine/Plugins/Online/EOSShared/Source/EOSShared/Public/EOSShared.h
new file mode 100644
--- /dev/null
+++ b/Engine/Plugins/Online/EOSShared/Source/EOSShared/Public/EOSShared.h
@@ -0,0 +1,21 @@
+// Copyright Epic Games, Inc. All Rights Reserved.
+
+#pragma once
+
+#if WITH_EOS_SDK
+
+#include "Containers/UnrealString.h"
+#include "Logging/LogMacros.h"
+
+#if defined(EOS_PLATFORM_BASE_FILE_NAME)
+#include EOS_PLATFORM_BASE_FILE_NAME
+#endif
+#include "eos_common.h"
+
+DECLARE_LOG_CATEGORY_EXTERN(LogEOSSDK, Log, All);
+
+inline FString LexToString(EOS_EResult EosResult)
+{
+	return ANSI_TO_TCHAR(EOS_EResult_ToString(EosResult));
+}
+#endif // WITH_EOS_SDK
\ No newline at end of file
diff --git a/Engine/Plugins/Online/EOSShared/Source/EOSShared/Public/EOSSharedSettings.h b/Engine/Plugins/Online/EOSShared/Source/EOSShared/Public/EOSSharedSettings.h
new file mode 100644
--- /dev/null
+++ b/Engine/Plugins/Online/EOSShared/Source/EOSShared/Public/EOSSharedSettings.h
@@ -0,0 +1,54 @@
+// Copyright Epic Games, Inc. All Rights Reserved.
+
+#pragma once
+
+#include "UObject/ObjectMacros.h"
+#include "UObject/Object.h"
+#include "EOSSharedSettings.generated.h"
+
+UENUM(BlueprintType)
+enum class EEosLogLevel : uint8
+{
+	Off,
+	Fatal,
+	Error,
+	Warning,
+	Info,
+	Verbose,
+	VeryVerbose
+};
+
+void LexFromString(EEosLogLevel& Value, const TCHAR* String);
+
+enum class EOS_ELogLevel : int32_t;
+EOS_ELogLevel ConvertLogLevel(const EEosLogLevel LogLevel);
+
+/** Native version of the UObject based config data */
+struct FEOSSharedSettings
+{
+	EEosLogLevel LogLevel = EEosLogLevel::Info;
+};
+
+UCLASS(Config=Engine, DefaultConfig)
+class EOSSHARED_API UEOSSharedSettings :
+	public UObject
+{
+	GENERATED_BODY()
+
+public:
+
+	/** Set the desired log level */
+	UPROPERTY(Config, EditAnywhere, BlueprintReadOnly, Category = "EOS Settings")
+	EEosLogLevel LogLevel = EEosLogLevel::Info;
+
+	static FEOSSharedSettings GetSettings();
+	FEOSSharedSettings ToNative() const;
+
+private:
+#if WITH_EDITOR
+	virtual void PostEditChangeProperty(FPropertyChangedEvent& PropertyChangedEvent) override;
+#endif
+
+	static FEOSSharedSettings AutoGetSettings();
+	static FEOSSharedSettings ManualGetSettings();
+};
diff --git a/Engine/Plugins/Online/EOSShared/Source/EOSShared/Public/EOSSharedUtils.h b/Engine/Plugins/Online/EOSShared/Source/EOSShared/Public/EOSSharedUtils.h
deleted file mode 100644
diff --git a/Engine/Plugins/Online/EOSShared/Source/EOSShared/Public/IEOSSDKManager.h b/Engine/Plugins/Online/EOSShared/Source/EOSShared/Public/IEOSSDKManager.h
--- a/Engine/Plugins/Online/EOSShared/Source/EOSShared/Public/IEOSSDKManager.h
+++ b/Engine/Plugins/Online/EOSShared/Source/EOSShared/Public/IEOSSDKManager.h
@@ -4,25 +4,50 @@
 
 #if WITH_EOS_SDK
 
-#include "Delegates/Delegate.h"
-#include "Features/IModularFeature.h"
-#include "Templates/Function.h"
+#include "Features/IModularFeatures.h"
 
 #if defined(EOS_PLATFORM_BASE_FILE_NAME)
 #include EOS_PLATFORM_BASE_FILE_NAME
 #endif
 #include "eos_types.h"
 
+struct IEOSPlatformHandle
+{
+	EOS_HPlatform PlatformHandle;
+
+	IEOSPlatformHandle(EOS_HPlatform InPlatformHandle) : PlatformHandle(InPlatformHandle) {}
+	virtual ~IEOSPlatformHandle() = default;
+
+	virtual void Tick() = 0;
+
+	operator EOS_HPlatform() const { return PlatformHandle; }
+};
+using IEOSPlatformHandlePtr = TSharedPtr<IEOSPlatformHandle, ESPMode::ThreadSafe>;
+
 class IEOSSDKManager : public IModularFeature
 {
 public:
+	static IEOSSDKManager* Get()
+	{
+		if (IModularFeatures::Get().IsModularFeatureAvailable(GetModularFeatureName()))
+		{
+			return &IModularFeatures::Get().GetModularFeature<IEOSSDKManager>(GetModularFeatureName());
+		}
+		return nullptr;
+	}
+
+	static FName GetModularFeatureName()
+	{
+		static const FName FeatureName = TEXT("EOSSDKManager");
+		return FeatureName;
+	}
+
+	virtual ~IEOSSDKManager() = default;
+
 	virtual EOS_EResult Initialize() = 0;
 	virtual bool IsInitialized() const = 0;
 
-	virtual EOS_HPlatform CreatePlatform(const EOS_Platform_Options& PlatformOptions) = 0;
-	virtual void ReleasePlatform(EOS_HPlatform PlatformHandle) = 0;
-
-	virtual bool Tick(float DeltaTime = 0.f) = 0;
+	virtual IEOSPlatformHandlePtr CreatePlatform(const EOS_Platform_Options& PlatformOptions) = 0;
 
 	virtual FString GetProductName() const = 0;
 	virtual FString GetProductVersion() const = 0;
diff --git a/Engine/Plugins/Online/OnlineSubsystemEOS/ReadMe.txt b/Engine/Plugins/Online/OnlineSubsystemEOS/ReadMe.txt
--- a/Engine/Plugins/Online/OnlineSubsystemEOS/ReadMe.txt
+++ b/Engine/Plugins/Online/OnlineSubsystemEOS/ReadMe.txt
@@ -19,6 +19,6 @@ This plugin supports:
 To configure your settings:
 	1. Go to Project Settings
 	2. Scroll down to Plugins
-	3. Select Epic Online Services (see EOS Plugin Settings.jpg)
+	3. Select Epic Online Services
 	4. Add a entry in the list for each artifact on EGS or use the default name for non-EGS
 		- Each deployment id/artifact can have separate settings
diff --git a/Engine/Plugins/Online/OnlineSubsystemEOS/Source/OnlineSubsystemEOS/Private/EOSSettings.cpp b/Engine/Plugins/Online/OnlineSubsystemEOS/Source/OnlineSubsystemEOS/Private/EOSSettings.cpp
--- a/Engine/Plugins/Online/OnlineSubsystemEOS/Source/OnlineSubsystemEOS/Private/EOSSettings.cpp
+++ b/Engine/Plugins/Online/OnlineSubsystemEOS/Source/OnlineSubsystemEOS/Private/EOSSettings.cpp
@@ -4,6 +4,7 @@
 #include "OnlineSubsystemEOS.h"
 #include "OnlineSubsystemEOSModule.h"
 
+#include "Algo/Transform.h"
 #include "Misc/CommandLine.h"
 #include "Misc/ConfigCacheIni.h"
 
@@ -66,7 +67,7 @@ FEOSArtifactSettings FArtifactSettings::ToNative() const
 	return Native;
 }
 
-inline FString StripQuotes(FString& Source)
+inline FString StripQuotes(const FString& Source)
 {
 	if (Source.StartsWith(TEXT("\"")))
 	{
@@ -75,7 +76,7 @@ inline FString StripQuotes(FString& Source)
 	return Source;
 }
 
-void FEOSArtifactSettings::ParseRawArrayEntry(FString& RawLine)
+void FEOSArtifactSettings::ParseRawArrayEntry(const FString& RawLine)
 {
 	const TCHAR* Delims[4] = { TEXT("("), TEXT(")"), TEXT("="), TEXT(",") };
 	TArray<FString> Values;
@@ -120,12 +121,6 @@ void FEOSArtifactSettings::ParseRawArrayEntry(FString& RawLine)
 	}
 }
 
-UEOSSettings::UEOSSettings()
-	: CacheDir(TEXT("CacheDir"))
-{
-
-}
-
 FEOSSettings UEOSSettings::GetSettings()
 {
 	if (UObjectInitialized())
@@ -144,21 +139,21 @@ FEOSSettings UEOSSettings::ManualGetSettings()
 {
 	FEOSSettings Native;
 
-	GConfig->GetBool(INI_SECTION, TEXT("bEnableOverlay"), Native.bEnableOverlay, GEngineIni);
-	GConfig->GetBool(INI_SECTION, TEXT("bEnableSocialOverlay"), Native.bEnableSocialOverlay, GEngineIni);
-	GConfig->GetBool(INI_SECTION, TEXT("bShouldEnforceBeingLaunchedByEGS"), Native.bShouldEnforceBeingLaunchedByEGS, GEngineIni);
 	GConfig->GetString(INI_SECTION, TEXT("CacheDir"), Native.CacheDir, GEngineIni);
 	GConfig->GetString(INI_SECTION, TEXT("DefaultArtifactName"), Native.DefaultArtifactName, GEngineIni);
-	GConfig->GetArray(INI_SECTION, TEXT("TitleStorageTags"), Native.TitleStorageTags, GEngineIni);
 	GConfig->GetInt(INI_SECTION, TEXT("TickBudgetInMilliseconds"), Native.TickBudgetInMilliseconds, GEngineIni);
 	GConfig->GetInt(INI_SECTION, TEXT("TitleStorageReadChunkLength"), Native.TitleStorageReadChunkLength, GEngineIni);
-
+	GConfig->GetBool(INI_SECTION, TEXT("bEnableOverlay"), Native.bEnableOverlay, GEngineIni);
+	GConfig->GetBool(INI_SECTION, TEXT("bEnableSocialOverlay"), Native.bEnableSocialOverlay, GEngineIni);
+	GConfig->GetBool(INI_SECTION, TEXT("bShouldEnforceBeingLaunchedByEGS"), Native.bShouldEnforceBeingLaunchedByEGS, GEngineIni);
 	GConfig->GetBool(INI_SECTION, TEXT("bUseEAS"), Native.bUseEAS, GEngineIni);
 	GConfig->GetBool(INI_SECTION, TEXT("bUseEOSConnect"), Native.bUseEOSConnect, GEngineIni);
+	GConfig->GetBool(INI_SECTION, TEXT("bUseEOSSessions"), Native.bUseEOSSessions, GEngineIni);
 	GConfig->GetBool(INI_SECTION, TEXT("bMirrorStatsToEOS"), Native.bMirrorStatsToEOS, GEngineIni);
 	GConfig->GetBool(INI_SECTION, TEXT("bMirrorAchievementsToEOS"), Native.bMirrorAchievementsToEOS, GEngineIni);
-	GConfig->GetBool(INI_SECTION, TEXT("bUseEOSSessions"), Native.bUseEOSSessions, GEngineIni);
 	GConfig->GetBool(INI_SECTION, TEXT("bMirrorPresenceToEAS"), Native.bMirrorPresenceToEAS, GEngineIni);
+	// Artifacts explicitly skipped
+	GConfig->GetArray(INI_SECTION, TEXT("TitleStorageTags"), Native.TitleStorageTags, GEngineIni);
 
 	return Native;
 }
@@ -167,18 +162,21 @@ FEOSSettings UEOSSettings::ToNative() const
 {
 	FEOSSettings Native;
 
-	Native.bEnableOverlay = bEnableOverlay;
-	Native.bEnableSocialOverlay = bEnableSocialOverlay;
 	Native.CacheDir = CacheDir;
 	Native.DefaultArtifactName = DefaultArtifactName;
-	Native.TitleStorageTags = TitleStorageTags;
 	Native.TickBudgetInMilliseconds = TickBudgetInMilliseconds;
 	Native.TitleStorageReadChunkLength = TitleStorageReadChunkLength;
-
-	for (const FArtifactSettings& ArtifactSettings : Artifacts)
-	{
-		Native.Artifacts.Add(ArtifactSettings.ToNative());
-	}
+	Native.bEnableOverlay = bEnableOverlay;
+	Native.bEnableSocialOverlay = bEnableSocialOverlay;
+	Native.bShouldEnforceBeingLaunchedByEGS = bShouldEnforceBeingLaunchedByEGS;
+	Native.bUseEAS = bUseEAS;
+	Native.bUseEOSConnect = bUseEOSConnect;
+	Native.bUseEOSSessions = bUseEOSSessions;
+	Native.bMirrorStatsToEOS = bMirrorStatsToEOS;
+	Native.bMirrorAchievementsToEOS = bMirrorAchievementsToEOS;
+	Native.bMirrorPresenceToEAS = bMirrorPresenceToEAS;
+	Algo::Transform(Artifacts, Native.Artifacts, &FArtifactSettings::ToNative);
+	Native.TitleStorageTags = TitleStorageTags;
 
 	return Native;
 }
@@ -201,7 +199,7 @@ bool UEOSSettings::ManualGetSettingsForArtifact(const FString& ArtifactName, FEO
 
 	TArray<FString> Artifacts;
 	GConfig->GetArray(INI_SECTION, TEXT("Artifacts"), Artifacts, GEngineIni);
-	for (FString& Line : Artifacts)
+	for (const FString& Line : Artifacts)
 	{
 		FEOSArtifactSettings Artifact;
 		Artifact.ParseRawArrayEntry(Line);
@@ -216,7 +214,7 @@ bool UEOSSettings::ManualGetSettingsForArtifact(const FString& ArtifactName, FEO
 		ArtifactNameOverride = ArtifactName;
 	}
 	// Search by name and then default if not found
-	for (FEOSArtifactSettings& Artifact : ArtifactSettings)
+	for (const FEOSArtifactSettings& Artifact : ArtifactSettings)
 	{
 		if (Artifact.ArtifactName == ArtifactNameOverride)
 		{
@@ -224,7 +222,7 @@ bool UEOSSettings::ManualGetSettingsForArtifact(const FString& ArtifactName, FEO
 			return true;
 		}
 	}
-	for (FEOSArtifactSettings& Artifact : ArtifactSettings)
+	for (const FEOSArtifactSettings& Artifact : ArtifactSettings)
 	{
 		if (Artifact.ArtifactName == DefaultArtifactName)
 		{
diff --git a/Engine/Plugins/Online/OnlineSubsystemEOS/Source/OnlineSubsystemEOS/Private/OnlineSessionEOS.cpp b/Engine/Plugins/Online/OnlineSubsystemEOS/Source/OnlineSubsystemEOS/Private/OnlineSessionEOS.cpp
--- a/Engine/Plugins/Online/OnlineSubsystemEOS/Source/OnlineSubsystemEOS/Private/OnlineSessionEOS.cpp
+++ b/Engine/Plugins/Online/OnlineSubsystemEOS/Source/OnlineSubsystemEOS/Private/OnlineSessionEOS.cpp
@@ -11,6 +11,7 @@
 #include "SocketSubsystem.h"
 #include "NboSerializerEOS.h"
 #include "InternetAddrEOS.h"
+#include "IEOSSDKManager.h"
 
 #if WITH_EOS_SDK
 	#include "eos_sessions.h"
@@ -470,7 +471,7 @@ void FOnlineSessionEOS::Init(const FString& InBucketId)
 	SessionInviteAcceptedId = EOS_Sessions_AddNotifySessionInviteAccepted(EOSSubsystem->SessionsHandle, &Options, SessionInviteAcceptedCallbackObj, SessionInviteAcceptedCallbackObj->GetCallbackPtr());
 
 	// Lobbies
-	LobbyHandle = EOS_Platform_GetLobbyInterface(EOSSubsystem->EOSPlatformHandle);
+	LobbyHandle = EOS_Platform_GetLobbyInterface(*EOSSubsystem->EOSPlatformHandle);
 	RegisterLobbyNotifications();
 
 	bIsDedicatedServer = IsRunningDedicatedServer();
diff --git a/Engine/Plugins/Online/OnlineSubsystemEOS/Source/OnlineSubsystemEOS/Private/OnlineSubsystemEOS.cpp b/Engine/Plugins/Online/OnlineSubsystemEOS/Source/OnlineSubsystemEOS/Private/OnlineSubsystemEOS.cpp
--- a/Engine/Plugins/Online/OnlineSubsystemEOS/Source/OnlineSubsystemEOS/Private/OnlineSubsystemEOS.cpp
+++ b/Engine/Plugins/Online/OnlineSubsystemEOS/Source/OnlineSubsystemEOS/Private/OnlineSubsystemEOS.cpp
@@ -1,6 +1,7 @@
 // Copyright Epic Games, Inc. All Rights Reserved.
 
 #include "OnlineSubsystemEOS.h"
+#include "OnlineSubsystemUtils.h"
 #include "UserManagerEOS.h"
 #include "OnlineSessionEOS.h"
 #include "OnlineStatsEOS.h"
@@ -24,18 +25,6 @@
 #define EOS_ENCRYPTION_KEY_MAX_LENGTH 64
 #define EOS_ENCRYPTION_KEY_MAX_BUFFER_LEN (EOS_ENCRYPTION_KEY_MAX_LENGTH + 1)
 
-namespace {
-	IEOSSDKManager* GetEOSSDKManager()
-	{
-		const FName EOSSDKManagerFeatureName = TEXT("EOSSDKManager");
-		if (IModularFeatures::Get().IsModularFeatureAvailable(EOSSDKManagerFeatureName))
-		{
-			return &IModularFeatures::Get().GetModularFeature<IEOSSDKManager>(EOSSDKManagerFeatureName);
-		}
-		return nullptr;
-	}
-}
-
 /** Class that holds the strings for the call duration */
 struct FEOSPlatformOptions :
 	public EOS_Platform_Options
@@ -71,9 +60,9 @@ void FOnlineSubsystemEOS::ModuleInit()
 	const FName EOSSharedModuleName = TEXT("EOSShared");
 	if (!FModuleManager::Get().IsModuleLoaded(EOSSharedModuleName))
 	{
-		FModuleManager::Get().LoadModule(EOSSharedModuleName);
+		FModuleManager::Get().LoadModuleChecked(EOSSharedModuleName);
 	}
-	IEOSSDKManager* EOSSDKManager = GetEOSSDKManager();
+	IEOSSDKManager* EOSSDKManager = IEOSSDKManager::Get();
 	if (!EOSSDKManager)
 	{
 		UE_LOG_ONLINE(Error, TEXT("FOnlineSubsystemEOS: Missing IEOSSDKManager modular feature."));
@@ -114,11 +103,10 @@ bool FOnlineSubsystemEOS::PlatformCreate()
 		UE_LOG_ONLINE(Error, TEXT("FOnlineSubsystemEOS::PlatformCreate() failed to find artifact settings object for artifact (%s)"), *ArtifactName);
 		return false;
 	}
-
-	EOSSDKManager = GetEOSSDKManager();
+	EOSSDKManager = IEOSSDKManager::Get();
 	if (!EOSSDKManager)
 	{
-		UE_LOG_ONLINE(Error, TEXT("FOnlineSubsystemEOS::PlatformCreate() failed to init EOS platform"));
+		UE_LOG_ONLINE(Error, TEXT("FOnlineSubsystemEOS::PlatformCreate() failed to get EOSSDKManager interface"));
 		return false;
 	}
 
@@ -187,73 +175,73 @@ bool FOnlineSubsystemEOS::Init()
 	}
 
 	// Get handles for later use
-	AuthHandle = EOS_Platform_GetAuthInterface(EOSPlatformHandle);
+	AuthHandle = EOS_Platform_GetAuthInterface(*EOSPlatformHandle);
 	if (AuthHandle == nullptr)
 	{
 		UE_LOG_ONLINE(Error, TEXT("FOnlineSubsystemEOS: failed to init EOS platform, couldn't get auth handle"));
 		return false;
 	}
-	UserInfoHandle = EOS_Platform_GetUserInfoInterface(EOSPlatformHandle);
+	UserInfoHandle = EOS_Platform_GetUserInfoInterface(*EOSPlatformHandle);
 	if (UserInfoHandle == nullptr)
 	{
 		UE_LOG_ONLINE(Error, TEXT("FOnlineSubsystemEOS: failed to init EOS platform, couldn't get user info handle"));
 		return false;
 	}
-	UIHandle = EOS_Platform_GetUIInterface(EOSPlatformHandle);
+	UIHandle = EOS_Platform_GetUIInterface(*EOSPlatformHandle);
 	if (UIHandle == nullptr)
 	{
 		UE_LOG_ONLINE(Error, TEXT("FOnlineSubsystemEOS: failed to init EOS platform, couldn't get UI handle"));
 		return false;
 	}
-	FriendsHandle = EOS_Platform_GetFriendsInterface(EOSPlatformHandle);
+	FriendsHandle = EOS_Platform_GetFriendsInterface(*EOSPlatformHandle);
 	if (FriendsHandle == nullptr)
 	{
 		UE_LOG_ONLINE(Error, TEXT("FOnlineSubsystemEOS: failed to init EOS platform, couldn't get friends handle"));
 		return false;
 	}
-	PresenceHandle = EOS_Platform_GetPresenceInterface(EOSPlatformHandle);
+	PresenceHandle = EOS_Platform_GetPresenceInterface(*EOSPlatformHandle);
 	if (PresenceHandle == nullptr)
 	{
 		UE_LOG_ONLINE(Error, TEXT("FOnlineSubsystemEOS: failed to init EOS platform, couldn't get presence handle"));
 		return false;
 	}
-	ConnectHandle = EOS_Platform_GetConnectInterface(EOSPlatformHandle);
+	ConnectHandle = EOS_Platform_GetConnectInterface(*EOSPlatformHandle);
 	if (ConnectHandle == nullptr)
 	{
 		UE_LOG_ONLINE(Error, TEXT("FOnlineSubsystemEOS: failed to init EOS platform, couldn't get connect handle"));
 		return false;
 	}
-	SessionsHandle = EOS_Platform_GetSessionsInterface(EOSPlatformHandle);
+	SessionsHandle = EOS_Platform_GetSessionsInterface(*EOSPlatformHandle);
 	if (SessionsHandle == nullptr)
 	{
 		UE_LOG_ONLINE(Error, TEXT("FOnlineSubsystemEOS: failed to init EOS platform, couldn't get sessions handle"));
 		return false;
 	}
-	StatsHandle = EOS_Platform_GetStatsInterface(EOSPlatformHandle);
+	StatsHandle = EOS_Platform_GetStatsInterface(*EOSPlatformHandle);
 	if (StatsHandle == nullptr)
 	{
 		UE_LOG_ONLINE(Error, TEXT("FOnlineSubsystemEOS: failed to init EOS platform, couldn't get stats handle"));
 		return false;
 	}
-	LeaderboardsHandle = EOS_Platform_GetLeaderboardsInterface(EOSPlatformHandle);
+	LeaderboardsHandle = EOS_Platform_GetLeaderboardsInterface(*EOSPlatformHandle);
 	if (LeaderboardsHandle == nullptr)
 	{
 		UE_LOG_ONLINE(Error, TEXT("FOnlineSubsystemEOS: failed to init EOS platform, couldn't get leaderboards handle"));
 		return false;
 	}
-	MetricsHandle = EOS_Platform_GetMetricsInterface(EOSPlatformHandle);
+	MetricsHandle = EOS_Platform_GetMetricsInterface(*EOSPlatformHandle);
 	if (MetricsHandle == nullptr)
 	{
 		UE_LOG_ONLINE(Error, TEXT("FOnlineSubsystemEOS: failed to init EOS platform, couldn't get metrics handle"));
 		return false;
 	}
-	AchievementsHandle = EOS_Platform_GetAchievementsInterface(EOSPlatformHandle);
+	AchievementsHandle = EOS_Platform_GetAchievementsInterface(*EOSPlatformHandle);
 	if (AchievementsHandle == nullptr)
 	{
 		UE_LOG_ONLINE(Error, TEXT("FOnlineSubsystemEOS: failed to init EOS platform, couldn't get achievements handle"));
 		return false;
 	}
-	P2PHandle = EOS_Platform_GetP2PInterface(EOSPlatformHandle);
+	P2PHandle = EOS_Platform_GetP2PInterface(*EOSPlatformHandle);
 	if (P2PHandle == nullptr)
 	{
 		UE_LOG_ONLINE(Error, TEXT("FOnlineSubsystemEOS: failed to init EOS platform, couldn't get p2p handle"));
@@ -262,7 +250,7 @@ bool FOnlineSubsystemEOS::Init()
 	// Disable ecom if not part of EGS
 	if (bWasLaunchedByEGS)
 	{
-		EcomHandle = EOS_Platform_GetEcomInterface(EOSPlatformHandle);
+		EcomHandle = EOS_Platform_GetEcomInterface(*EOSPlatformHandle);
 		if (EcomHandle == nullptr)
 		{
 			UE_LOG_ONLINE(Error, TEXT("FOnlineSubsystemEOS: failed to init EOS platform, couldn't get ecom handle"));
@@ -270,13 +258,13 @@ bool FOnlineSubsystemEOS::Init()
 		}
 		StoreInterfacePtr = MakeShareable(new FOnlineStoreEOS(this));
 	}
-	TitleStorageHandle = EOS_Platform_GetTitleStorageInterface(EOSPlatformHandle);
+	TitleStorageHandle = EOS_Platform_GetTitleStorageInterface(*EOSPlatformHandle);
 	if (TitleStorageHandle == nullptr)
 	{
 		UE_LOG_ONLINE(Error, TEXT("FOnlineSubsystemEOS: failed to init EOS platform, couldn't get title storage handle"));
 		return false;
 	}
-	PlayerDataStorageHandle = EOS_Platform_GetPlayerDataStorageInterface(EOSPlatformHandle);
+	PlayerDataStorageHandle = EOS_Platform_GetPlayerDataStorageInterface(*EOSPlatformHandle);
 	if (PlayerDataStorageHandle == nullptr)
 	{
 		UE_LOG_ONLINE(Error, TEXT("FOnlineSubsystemEOS: failed to init EOS platform, couldn't get player data storage handle"));
@@ -380,6 +368,38 @@ FText FOnlineSubsystemEOS::GetOnlineServiceName() const
 	return NSLOCTEXT("OnlineSubsystemEOS", "OnlineServiceName", "EOS");
 }
 
+FOnlineSubsystemEOS::FOnlineSubsystemEOS(FName InInstanceName) :
+	FOnlineSubsystemImpl(EOS_SUBSYSTEM, InInstanceName)
+	, EOSSDKManager(nullptr)
+	, AuthHandle(nullptr)
+	, UIHandle(nullptr)
+	, FriendsHandle(nullptr)
+	, UserInfoHandle(nullptr)
+	, PresenceHandle(nullptr)
+	, ConnectHandle(nullptr)
+	, SessionsHandle(nullptr)
+	, StatsHandle(nullptr)
+	, LeaderboardsHandle(nullptr)
+	, MetricsHandle(nullptr)
+	, AchievementsHandle(nullptr)
+	, P2PHandle(nullptr)
+	, EcomHandle(nullptr)
+	, TitleStorageHandle(nullptr)
+	, PlayerDataStorageHandle(nullptr)
+	, UserManager(nullptr)
+	, SessionInterfacePtr(nullptr)
+	, LeaderboardsInterfacePtr(nullptr)
+	, AchievementsInterfacePtr(nullptr)
+	, StoreInterfacePtr(nullptr)
+	, TitleFileInterfacePtr(nullptr)
+	, UserCloudInterfacePtr(nullptr)
+	, bWasLaunchedByEGS(false)
+	, bIsDefaultOSS(false)
+	, bIsPlatformOSS(false)
+{
+	StopTicker();
+}
+
 IOnlineSessionPtr FOnlineSubsystemEOS::GetSessionInterface() const
 {
 	return SessionInterfacePtr;
diff --git a/Engine/Plugins/Online/OnlineSubsystemEOS/Source/OnlineSubsystemEOS/Private/OnlineSubsystemEOS.h b/Engine/Plugins/Online/OnlineSubsystemEOS/Source/OnlineSubsystemEOS/Private/OnlineSubsystemEOS.h
--- a/Engine/Plugins/Online/OnlineSubsystemEOS/Source/OnlineSubsystemEOS/Private/OnlineSubsystemEOS.h
+++ b/Engine/Plugins/Online/OnlineSubsystemEOS/Source/OnlineSubsystemEOS/Private/OnlineSubsystemEOS.h
@@ -16,6 +16,7 @@ DECLARE_STATS_GROUP(TEXT("EOS"), STATGROUP_EOS, STATCAT_Advanced);
 #include "eos_sdk.h"
 
 class IEOSSDKManager;
+using IEOSPlatformHandlePtr = TSharedPtr<struct IEOSPlatformHandle, ESPMode::ThreadSafe>;
 
 class FUserManagerEOS;
 typedef TSharedPtr<class FUserManagerEOS, ESPMode::ThreadSafe> FUserManagerEOSPtr;
@@ -108,38 +109,7 @@ public:
 PACKAGE_SCOPE:
 	/** Only the factory makes instances */
 	FOnlineSubsystemEOS() = delete;
-	explicit FOnlineSubsystemEOS(FName InInstanceName) :
-		FOnlineSubsystemImpl(EOS_SUBSYSTEM, InInstanceName)
-		, EOSSDKManager(nullptr)
-		, EOSPlatformHandle(nullptr)
-		, AuthHandle(nullptr)
-		, UIHandle(nullptr)
-		, FriendsHandle(nullptr)
-		, UserInfoHandle(nullptr)
-		, PresenceHandle(nullptr)
-		, ConnectHandle(nullptr)
-		, SessionsHandle(nullptr)
-		, StatsHandle(nullptr)
-		, LeaderboardsHandle(nullptr)
-		, MetricsHandle(nullptr)
-		, AchievementsHandle(nullptr)
-		, P2PHandle(nullptr)
-		, EcomHandle(nullptr)
-		, TitleStorageHandle(nullptr)
-		, PlayerDataStorageHandle(nullptr)
-		, UserManager(nullptr)
-		, SessionInterfacePtr(nullptr)
-		, LeaderboardsInterfacePtr(nullptr)
-		, AchievementsInterfacePtr(nullptr)
-		, StoreInterfacePtr(nullptr)
-		, TitleFileInterfacePtr(nullptr)
-		, UserCloudInterfacePtr(nullptr)
-		, bWasLaunchedByEGS(false)
-		, bIsDefaultOSS(false)
-		, bIsPlatformOSS(false)
-	{
-		StopTicker();
-	}
+	explicit FOnlineSubsystemEOS(FName InInstanceName);
 
 	char ProductNameAnsi[EOS_PRODUCTNAME_MAX_BUFFER_LEN];
 	char ProductVersionAnsi[EOS_PRODUCTVERSION_MAX_BUFFER_LEN];
@@ -147,7 +117,7 @@ PACKAGE_SCOPE:
 	IEOSSDKManager* EOSSDKManager;
 
 	/** EOS handles */
-	EOS_HPlatform EOSPlatformHandle;
+	IEOSPlatformHandlePtr EOSPlatformHandle;
 	EOS_HAuth AuthHandle;
 	EOS_HUI UIHandle;
 	EOS_HFriends FriendsHandle;
diff --git a/Engine/Plugins/Online/OnlineSubsystemEOS/Source/OnlineSubsystemEOS/Private/OnlineSubsystemModuleEOS.cpp b/Engine/Plugins/Online/OnlineSubsystemEOS/Source/OnlineSubsystemEOS/Private/OnlineSubsystemModuleEOS.cpp
--- a/Engine/Plugins/Online/OnlineSubsystemEOS/Source/OnlineSubsystemEOS/Private/OnlineSubsystemModuleEOS.cpp
+++ b/Engine/Plugins/Online/OnlineSubsystemEOS/Source/OnlineSubsystemEOS/Private/OnlineSubsystemModuleEOS.cpp
@@ -72,12 +72,11 @@ void FOnlineSubsystemEOSModule::StartupModule()
 #if WITH_EDITOR
 void FOnlineSubsystemEOSModule::OnPostEngineInit()
 {
-	ISettingsModule* SettingsModule = FModuleManager::GetModulePtr<ISettingsModule>("Settings");
-	if (SettingsModule != nullptr)
+	if (ISettingsModule* SettingsModule = FModuleManager::GetModulePtr<ISettingsModule>("Settings"))
 	{
-		SettingsModule->RegisterSettings("Project", "Plugins", "Epic Online Services",
-			LOCTEXT("EOSSettingsName", "Epic Online Services"),
-			LOCTEXT("EOSSettingsDescription", "Configure the Epic Online Services"),
+		SettingsModule->RegisterSettings("Project", "Plugins", "Online Subsystem EOS",
+			LOCTEXT("OSSEOSSettingsName", "Online Subsystem EOS"),
+			LOCTEXT("OSSEOSSettingsDescription", "Configure Online Subsystem EOS"),
 			GetMutableDefault<UEOSSettings>());
 	}
 }
@@ -86,19 +85,19 @@ void FOnlineSubsystemEOSModule::OnPostEngineInit()
 #if WITH_EDITOR
 void FOnlineSubsystemEOSModule::OnPreExit()
 {
-	ISettingsModule* SettingsModule = FModuleManager::GetModulePtr<ISettingsModule>("Settings");
-	if (SettingsModule)
+	if (ISettingsModule* SettingsModule = FModuleManager::GetModulePtr<ISettingsModule>("Settings"))
 	{
-		SettingsModule->UnregisterSettings("Project", "Plugins", "Epic Online Services");
+		SettingsModule->UnregisterSettings("Project", "Plugins", "Online Subsystem EOS");
 	}
 }
 #endif
 
 void FOnlineSubsystemEOSModule::ShutdownModule()
 {
-	FCoreDelegates::OnInit.RemoveAll(this);
+#if WITH_EDITOR
 	FCoreDelegates::OnPostEngineInit.RemoveAll(this);
 	FCoreDelegates::OnPreExit.RemoveAll(this);
+#endif
 
 #if WITH_EOS_SDK
 	FOnlineSubsystemEOS::ModuleShutdown();
diff --git a/Engine/Plugins/Online/OnlineSubsystemEOS/Source/OnlineSubsystemEOS/Public/EOSSettings.h b/Engine/Plugins/Online/OnlineSubsystemEOS/Source/OnlineSubsystemEOS/Public/EOSSettings.h
--- a/Engine/Plugins/Online/OnlineSubsystemEOS/Source/OnlineSubsystemEOS/Public/EOSSettings.h
+++ b/Engine/Plugins/Online/OnlineSubsystemEOS/Source/OnlineSubsystemEOS/Public/EOSSettings.h
@@ -18,7 +18,7 @@ struct FEOSArtifactSettings
 	FString DeploymentId;
 	FString EncryptionKey;
 
-	void ParseRawArrayEntry(FString& RawLine);
+	void ParseRawArrayEntry(const FString& RawLine);
 };
 
 UCLASS(Deprecated)
@@ -76,9 +76,9 @@ struct FEOSSettings
 	bool bShouldEnforceBeingLaunchedByEGS;
 	bool bUseEAS;
 	bool bUseEOSConnect;
+	bool bUseEOSSessions;
 	bool bMirrorStatsToEOS;
 	bool bMirrorAchievementsToEOS;
-	bool bUseEOSSessions;
 	bool bMirrorPresenceToEAS;
 	TArray<FEOSArtifactSettings> Artifacts;
 	TArray<FString> TitleStorageTags;
@@ -91,15 +91,13 @@ class ONLINESUBSYSTEMEOS_API UEOSSettings :
 	GENERATED_BODY()
 
 public:
-	UEOSSettings();
-
 	/**
 	 * The directory any PDS/TDS files are cached into. This is per artifact e.g.:
 	 *
 	 * <UserDir>/<ArtifactId>/<CacheDir>
 	 */
 	UPROPERTY(Config, EditAnywhere, BlueprintReadOnly, Category="EOS Settings")
-	FString CacheDir;
+	FString CacheDir = TEXT("CacheDir");
 
 	/** Used when launched from a store other than EGS or when the specified artifact name was not present */
 	UPROPERTY(Config, EditAnywhere, BlueprintReadOnly, Category="EOS Settings")
@@ -107,19 +105,19 @@ public:
 
 	/** Used to throttle how much time EOS ticking can take */
 	UPROPERTY(Config, EditAnywhere, BlueprintReadOnly, Category="EOS Settings")
-	int32 TickBudgetInMilliseconds;
+	int32 TickBudgetInMilliseconds = 0;
 
 	/** Set to true to enable the overlay (ecom features) */
 	UPROPERTY(Config, EditAnywhere, BlueprintReadOnly, Category="EOS Settings")
-	bool bEnableOverlay;
+	bool bEnableOverlay = false;
 
 	/** Set to true to enable the social overlay (friends, invites, etc.) */
 	UPROPERTY(Config, EditAnywhere, BlueprintReadOnly, Category="EOS Settings")
-	bool bEnableSocialOverlay;
+	bool bEnableSocialOverlay = false;
 
 	/** Set to true to enable the social overlay (friends, invites, etc.) */
 	UPROPERTY(Config, EditAnywhere, BlueprintReadOnly, Category="EOS Settings", DisplayName="Require Being Launched by the Epic Games Store")
-	bool bShouldEnforceBeingLaunchedByEGS;
+	bool bShouldEnforceBeingLaunchedByEGS = false;
 
 	/** Tag combinations for paged queries in title file enumerations, separate tags within groups using `+` */
 	UPROPERTY(Config, EditAnywhere, BlueprintReadOnly, Category="EOS Settings")
@@ -127,7 +125,7 @@ public:
 
 	/** Chunk size used when reading a title file */
 	UPROPERTY(Config, EditAnywhere, BlueprintReadOnly, Category="EOS Settings")
-	int32 TitleStorageReadChunkLength;
+	int32 TitleStorageReadChunkLength = 0;
 
 	/** Per artifact SDK settings. A game might have a FooStaging, FooQA, and public Foo artifact */
 	UPROPERTY(Config, EditAnywhere, BlueprintReadOnly, Category="EOS Settings")
@@ -135,27 +133,27 @@ public:
 
 	/** Set to true to have Epic Accounts used (friends list will be unified with the default platform) */
 	UPROPERTY(Config, EditAnywhere, BlueprintReadOnly, Category="Crossplay Settings", DisplayName="Use Epic Account Services")
-	bool bUseEAS;
+	bool bUseEAS = false;
 
 	/** Set to true to have EOS Connect APIs used to link accounts for crossplay */
 	UPROPERTY(Config, EditAnywhere, BlueprintReadOnly, Category="Crossplay Settings", DisplayName="Use Crossplatform User IDs")
-	bool bUseEOSConnect;
+	bool bUseEOSConnect = false;
 
 	/** Set to true to write stats to EOS as well as the default platform */
 	UPROPERTY(Config, EditAnywhere, BlueprintReadOnly, Category="Crossplay Settings")
-	bool bMirrorStatsToEOS;
+	bool bMirrorStatsToEOS = false;
 
 	/** Set to true to write achievement data to EOS as well as the default platform */
 	UPROPERTY(Config, EditAnywhere, BlueprintReadOnly, Category="Crossplay Settings")
-	bool bMirrorAchievementsToEOS;
+	bool bMirrorAchievementsToEOS = false;
 
 	/** Set to true to use EOS for session registration with data mirrored to the default platform */
 	UPROPERTY(Config, EditAnywhere, BlueprintReadOnly, Category="Crossplay Settings", DisplayName="Use Crossplay Sessions")
-	bool bUseEOSSessions;
+	bool bUseEOSSessions = false;
 
 	/** Set to true to have Epic Accounts presence information updated when the default platform is updated */
 	UPROPERTY(Config, EditAnywhere, BlueprintReadOnly, Category="Crossplay Settings")
-	bool bMirrorPresenceToEAS;
+	bool bMirrorPresenceToEAS = false;
 
 	/** Find the Settings for an artifact by name */
 	static bool GetSettingsForArtifact(const FString& ArtifactName, FEOSArtifactSettings& OutSettings);
