https://p4-swarm.epicgames.net/changes/16400150
Common changes to process EOS Login with external auth tokens in consoles

diff --git a/Engine/Plugins/Online/OnlineSubsystem/Source/Public/Interfaces/OnlineIdentityInterface.h b/Engine/Plugins/Online/OnlineSubsystem/Source/Public/Interfaces/OnlineIdentityInterface.h
--- a/Engine/Plugins/Online/OnlineSubsystem/Source/Public/Interfaces/OnlineIdentityInterface.h
+++ b/Engine/Plugins/Online/OnlineSubsystem/Source/Public/Interfaces/OnlineIdentityInterface.h
@@ -182,11 +182,22 @@ DECLARE_DELEGATE_TwoParams(FOnRevokeAuthTokenCompleteDelegate, const FUniqueNetI
 struct FExternalAuthToken
 {
 	TArray<uint8> TokenData;
+	FString TokenString;
 
 	inline bool HasTokenData() const
 	{
 		return TokenData.Num() > 0;
 	}
+
+	inline bool HasTokenString() const
+	{
+		return !TokenString.IsEmpty();
+	}
+
+	inline bool IsValid() const
+	{
+		return HasTokenData() || HasTokenString();
+	}
 };
 
 /**
diff --git a/Engine/Plugins/Online/OnlineSubsystemEOS/Source/OnlineSubsystemEOS/OnlineSubsystemEOS.Build.cs b/Engine/Plugins/Online/OnlineSubsystemEOS/Source/OnlineSubsystemEOS/OnlineSubsystemEOS.Build.cs
--- a/Engine/Plugins/Online/OnlineSubsystemEOS/Source/OnlineSubsystemEOS/OnlineSubsystemEOS.Build.cs
+++ b/Engine/Plugins/Online/OnlineSubsystemEOS/Source/OnlineSubsystemEOS/OnlineSubsystemEOS.Build.cs
@@ -28,5 +28,11 @@ public class OnlineSubsystemEOS : ModuleRules
 				"Json"
 			}
 		);
+
+		PrivateDefinitions.Add("USE_XBL_XSTS_TOKEN=" + (bUseXblXstsToken ? "1" : "0"));
+		PrivateDefinitions.Add("USE_PSN_ID_TOKEN=" + (bUsePsnIdToken ? "1" : "0"));
 	}
+
+	protected virtual bool bUseXblXstsToken { get { return false; } }
+	protected virtual bool bUsePsnIdToken { get { return false; } }
 }
diff --git a/Engine/Plugins/Online/OnlineSubsystemEOS/Source/OnlineSubsystemEOS/Private/EOSHelpers.h b/Engine/Plugins/Online/OnlineSubsystemEOS/Source/OnlineSubsystemEOS/Private/EOSHelpers.h
--- a/Engine/Plugins/Online/OnlineSubsystemEOS/Source/OnlineSubsystemEOS/Private/EOSHelpers.h
+++ b/Engine/Plugins/Online/OnlineSubsystemEOS/Source/OnlineSubsystemEOS/Private/EOSHelpers.h
@@ -4,6 +4,10 @@
 
 #if WITH_EOS_SDK
 
+#if defined(EOS_PLATFORM_BASE_FILE_NAME)
+#include EOS_PLATFORM_BASE_FILE_NAME
+#endif
+
 #include "eos_auth_types.h"
 #include "Interfaces/OnlineExternalUIInterface.h"
 
diff --git a/Engine/Plugins/Online/OnlineSubsystemEOS/Source/OnlineSubsystemEOS/Private/OnlineSubsystemEOSTypes.h b/Engine/Plugins/Online/OnlineSubsystemEOS/Source/OnlineSubsystemEOS/Private/OnlineSubsystemEOSTypes.h
--- a/Engine/Plugins/Online/OnlineSubsystemEOS/Source/OnlineSubsystemEOS/Private/OnlineSubsystemEOSTypes.h
+++ b/Engine/Plugins/Online/OnlineSubsystemEOS/Source/OnlineSubsystemEOS/Private/OnlineSubsystemEOSTypes.h
@@ -15,6 +15,10 @@
 #include "IPAddress.h"
 
 #if WITH_EOS_SDK
+	#if defined(EOS_PLATFORM_BASE_FILE_NAME)
+	#include EOS_PLATFORM_BASE_FILE_NAME
+	#endif
+
 	#include "eos_common.h"
 #endif
 
@@ -736,6 +740,8 @@ PACKAGE_SCOPE:
 
 	FOnlineSessionInfoEOS(const FString& InHostIp, const FString& InSessionId, EOS_HSessionDetails InSessionHandle);
 
+	FOnlineSessionInfoEOS(const FString& InHostIp, FUniqueNetIdEOSRef UniqueNetId, EOS_HSessionDetails InSessionHandle);
+
 	/**
 	 * Initialize LAN session
 	 */
diff --git a/Engine/Plugins/Online/OnlineSubsystemEOS/Source/OnlineSubsystemEOS/Private/SocketSubsystemEOS.h b/Engine/Plugins/Online/OnlineSubsystemEOS/Source/OnlineSubsystemEOS/Private/SocketSubsystemEOS.h
--- a/Engine/Plugins/Online/OnlineSubsystemEOS/Source/OnlineSubsystemEOS/Private/SocketSubsystemEOS.h
+++ b/Engine/Plugins/Online/OnlineSubsystemEOS/Source/OnlineSubsystemEOS/Private/SocketSubsystemEOS.h
@@ -6,6 +6,10 @@
 #include "SocketSubsystem.h"
 
 #if WITH_EOS_SDK
+	#if defined(EOS_PLATFORM_BASE_FILE_NAME)
+	#include EOS_PLATFORM_BASE_FILE_NAME
+	#endif
+
 	#include "eos_p2p_types.h"
 #endif
 
diff --git a/Engine/Plugins/Online/OnlineSubsystemEOS/Source/OnlineSubsystemEOS/Private/UserManagerEOS.cpp b/Engine/Plugins/Online/OnlineSubsystemEOS/Source/OnlineSubsystemEOS/Private/UserManagerEOS.cpp
--- a/Engine/Plugins/Online/OnlineSubsystemEOS/Source/OnlineSubsystemEOS/Private/UserManagerEOS.cpp
+++ b/Engine/Plugins/Online/OnlineSubsystemEOS/Source/OnlineSubsystemEOS/Private/UserManagerEOS.cpp
@@ -88,23 +88,30 @@ static inline EOS_Presence_EStatus ToEOS_Presence_EStatus(EOnlinePresenceState::
 	return EOS_Presence_EStatus::EOS_PS_Offline;
 }
 
-static inline EOS_EExternalCredentialType ToEOS_EExternalCredentialType(FName OSSName)
+static inline EOS_EExternalCredentialType ToEOS_EExternalCredentialType(FName OSSName, const FOnlineAccountCredentials& AccountCredentials)
 {
 	if (OSSName == STEAM_SUBSYSTEM)
 	{
 		return EOS_EExternalCredentialType::EOS_ECT_STEAM_APP_TICKET;
 	}
-	else if (OSSName == PS4_SUBSYSTEM)
+	else if (OSSName == PS4_SUBSYSTEM || USE_PSN_ID_TOKEN)
 	{
 		return EOS_EExternalCredentialType::EOS_ECT_PSN_ID_TOKEN;
 	}
-	else if (OSSName == LIVE_SUBSYSTEM)
+	else if (OSSName == LIVE_SUBSYSTEM || USE_XBL_XSTS_TOKEN)
 	{
 		return EOS_EExternalCredentialType::EOS_ECT_XBL_XSTS_TOKEN;
 	}
 	else if (OSSName == SWITCH_SUBSYSTEM)
 	{
-		return EOS_EExternalCredentialType::EOS_ECT_NINTENDO_ID_TOKEN;
+		if (AccountCredentials.Type == TEXT("NintendoAccount"))
+		{
+			return EOS_EExternalCredentialType::EOS_ECT_NINTENDO_ID_TOKEN;
+		}
+		else
+		{
+			return EOS_EExternalCredentialType::EOS_ECT_NINTENDO_NSA_ID_TOKEN;
+		}
 	}
 	else if (OSSName == APPLE_SUBSYSTEM)
 	{
@@ -259,8 +266,35 @@ struct FAuthCredentials :
 		Token = TokenAnsi;
 	}
 
-	FAuthCredentials(EOS_EExternalCredentialType InExternalType, const TArray<uint8>& InToken) :
+	FAuthCredentials(EOS_EExternalCredentialType InExternalType, const FExternalAuthToken& AuthToken) :
 		EOS_Auth_Credentials()
+	{
+		if (AuthToken.HasTokenData())
+		{
+			Init(InExternalType, AuthToken.TokenData);
+		}
+		else if (AuthToken.HasTokenString())
+		{
+			Init(InExternalType, AuthToken.TokenString);
+		}
+		else
+		{
+			UE_LOG_ONLINE(Error, TEXT("FAuthCredentials object cannot be constructed with invalid FExternalAuthToken parameter"));
+		}
+	}
+
+	void Init(EOS_EExternalCredentialType InExternalType, const FString& InTokenString)
+	{
+		ApiVersion = EOS_AUTH_CREDENTIALS_API_LATEST;
+		Type = EOS_ELoginCredentialType::EOS_LCT_ExternalAuth;
+		ExternalType = InExternalType;
+		Id = IdAnsi;
+		Token = TokenAnsi;
+
+		FCStringAnsi::Strncpy(TokenAnsi, TCHAR_TO_UTF8(*InTokenString), InTokenString.Len()+1);
+	}
+
+	void Init(EOS_EExternalCredentialType InExternalType, const TArray<uint8>& InToken)
 	{
 		ApiVersion = EOS_AUTH_CREDENTIALS_API_LATEST;
 		Type = EOS_ELoginCredentialType::EOS_LCT_ExternalAuth;
@@ -271,14 +305,19 @@ struct FAuthCredentials :
 		uint32_t InOutBufferLength = EOS_OSS_STRING_BUFFER_LENGTH;
 		EOS_ByteArray_ToString(InToken.GetData(), InToken.Num(), TokenAnsi, &InOutBufferLength);
 	}
+
 	char IdAnsi[EOS_OSS_STRING_BUFFER_LENGTH];
 	char TokenAnsi[EOS_MAX_TOKEN_SIZE];
 };
 
 bool FUserManagerEOS::Login(int32 LocalUserNum, const FOnlineAccountCredentials& AccountCredentials)
 {
+	LocalUserNumToLastLoginCredentials.Emplace(LocalUserNum, MakeShared<FOnlineAccountCredentials>(AccountCredentials));
+
+	FEOSSettings Settings = UEOSSettings::GetSettings();
+
 	// Are we configured to run at all?
-	if (!EOSSubsystem->bIsDefaultOSS && !EOSSubsystem->bIsPlatformOSS && !GetDefault<UEOSSettings>()->bUseEAS && !GetDefault<UEOSSettings>()->bUseEOSConnect)
+	if (!EOSSubsystem->bIsDefaultOSS && !EOSSubsystem->bIsPlatformOSS && !Settings.bUseEAS && !Settings.bUseEOSConnect)
 	{
 		UE_LOG_ONLINE(Warning, TEXT("Neither EAS nor EOS are configured to be used. Failed to login in user (%d)"), LocalUserNum);
 		TriggerOnLoginCompleteDelegates(LocalUserNum, false, *GetLocalUniqueNetIdEOS(LocalUserNum), FString(TEXT("Not configured")));
@@ -286,7 +325,7 @@ bool FUserManagerEOS::Login(int32 LocalUserNum, const FOnlineAccountCredentials&
 	}
 
 	// See if we are configured to just use EOS and not EAS
-	if (!EOSSubsystem->bIsDefaultOSS && !EOSSubsystem->bIsPlatformOSS && !GetDefault<UEOSSettings>()->bUseEAS && GetDefault<UEOSSettings>()->bUseEOSConnect)
+	if (!EOSSubsystem->bIsDefaultOSS && !EOSSubsystem->bIsPlatformOSS && !Settings.bUseEAS && Settings.bUseEOSConnect)
 	{
 		// Call the EOS + Platform login path
 		return ConnectLoginNoEAS(LocalUserNum);
@@ -301,7 +340,7 @@ bool FUserManagerEOS::Login(int32 LocalUserNum, const FOnlineAccountCredentials&
 	}
 
 	// See if we are logging in using platform credentials to link to EAS
-	if (!EOSSubsystem->bIsDefaultOSS && !EOSSubsystem->bIsPlatformOSS && GetDefault<UEOSSettings>()->bUseEAS)
+	if (!EOSSubsystem->bIsDefaultOSS && !EOSSubsystem->bIsPlatformOSS && Settings.bUseEAS)
 	{
 		LoginViaExternalAuth(LocalUserNum);
 		return true;
@@ -367,7 +406,7 @@ void FUserManagerEOS::LoginViaExternalAuth(int32 LocalUserNum)
 	GetPlatformAuthToken(LocalUserNum,
 		FOnGetLinkedAccountAuthTokenCompleteDelegate::CreateLambda([this](int32 LocalUserNum, bool bWasSuccessful, const FExternalAuthToken& AuthToken)
 		{
-			if (!bWasSuccessful || !AuthToken.HasTokenData())
+			if (!bWasSuccessful || !AuthToken.IsValid())
 			{
 				UE_LOG_ONLINE(Warning, TEXT("Unable to Login() user (%d) due to an empty platform auth token"), LocalUserNum);
 				TriggerOnLoginCompleteDelegates(LocalUserNum, false, *FUniqueNetIdEOS::EmptyId(), FString(TEXT("Missing platform auth token")));
@@ -378,7 +417,8 @@ void FUserManagerEOS::LoginViaExternalAuth(int32 LocalUserNum)
 			LoginOptions.ApiVersion = EOS_AUTH_LOGIN_API_LATEST;
 			LoginOptions.ScopeFlags = EOS_EAuthScopeFlags::EOS_AS_BasicProfile | EOS_EAuthScopeFlags::EOS_AS_FriendsList | EOS_EAuthScopeFlags::EOS_AS_Presence;
 
-			FAuthCredentials Credentials(ToEOS_EExternalCredentialType(GetPlatformOSS()->GetSubsystemName()), AuthToken.TokenData);
+			check(LocalUserNumToLastLoginCredentials.Contains(LocalUserNum));
+			FAuthCredentials Credentials(ToEOS_EExternalCredentialType(GetPlatformOSS()->GetSubsystemName(), *LocalUserNumToLastLoginCredentials[LocalUserNum]), AuthToken);
 			LoginOptions.Credentials = &Credentials;
 
 			FLoginCallback* CallbackObj = new FLoginCallback();
@@ -467,7 +507,8 @@ bool FUserManagerEOS::ConnectLoginNoEAS(int32 LocalUserNum)
 			}
 
 			// Now login into our EOS account
-			FConnectCredentials Credentials(ToEOS_EExternalCredentialType(GetPlatformOSS()->GetSubsystemName()), AuthToken.TokenData);
+			check(LocalUserNumToLastLoginCredentials.Contains(LocalUserNum));
+			FConnectCredentials Credentials(ToEOS_EExternalCredentialType(GetPlatformOSS()->GetSubsystemName(), *LocalUserNumToLastLoginCredentials[LocalUserNum]), AuthToken.TokenData);
 			EOS_Connect_LoginOptions Options = { };
 			Options.ApiVersion = EOS_CONNECT_LOGIN_API_LATEST;
 			Options.Credentials = &Credentials;
@@ -729,6 +770,8 @@ bool FUserManagerEOS::Logout(int32 LocalUserNum)
 
 	EOS_Auth_Logout(EOSSubsystem->AuthHandle, &LogoutOptions, CallbackObj, CallbackObj->GetCallbackPtr());
 
+	LocalUserNumToLastLoginCredentials.Remove(LocalUserNum);
+
 	return true;
 }
 
@@ -748,6 +791,9 @@ bool FUserManagerEOS::AutoLogin(int32 LocalUserNum)
 		return false;
 	}
 	FOnlineAccountCredentials Creds(AuthType, LoginId, Password);
+
+	LocalUserNumToLastLoginCredentials.Emplace(LocalUserNum, MakeShared<FOnlineAccountCredentials>(Creds));
+
 	return Login(LocalUserNum, Creds);
 }
 
@@ -1569,8 +1615,17 @@ void FUserManagerEOS::UpdateRemotePlayerProductUserId(EOS_EpicAccountId AccountI
 	StringToAccountIdMap.Remove(PrevNetIdStr);
 	StringToAccountIdMap.Add(NewNetIdStr, AccountId);
 	StringToProductUserIdMap.Add(NewNetIdStr, UserId);
-	FOnlineUserPtr OnlineUser = NetIdStringToOnlineUserMap[PrevNetIdStr];
-	NetIdStringToOnlineUserMap.Remove(PrevNetIdStr);
+	// If it's the first time we're updating this remote player, it will be indexed by its epic account id
+	FOnlineUserPtr OnlineUser = NetIdStringToOnlineUserMap[AccountIdStr];
+	if (OnlineUser.IsValid())
+	{
+		NetIdStringToOnlineUserMap.Remove(AccountIdStr);
+	}
+	else
+	{
+		OnlineUser = NetIdStringToOnlineUserMap[PrevNetIdStr];
+		NetIdStringToOnlineUserMap.Remove(PrevNetIdStr);
+	}
 	NetIdStringToOnlineUserMap.Add(NewNetIdStr, OnlineUser);
 	NetIdStringToAttributeAccessMap.Remove(PrevNetIdStr);
 	NetIdStringToAttributeAccessMap.Add(NewNetIdStr, AttrAccess);
diff --git a/Engine/Plugins/Online/OnlineSubsystemEOS/Source/OnlineSubsystemEOS/Private/UserManagerEOS.h b/Engine/Plugins/Online/OnlineSubsystemEOS/Source/OnlineSubsystemEOS/Private/UserManagerEOS.h
--- a/Engine/Plugins/Online/OnlineSubsystemEOS/Source/OnlineSubsystemEOS/Private/UserManagerEOS.h
+++ b/Engine/Plugins/Online/OnlineSubsystemEOS/Source/OnlineSubsystemEOS/Private/UserManagerEOS.h
@@ -439,6 +439,9 @@ private:
 	/** Identifier for the external UI notification callback */
 	EOS_NotificationId DisplaySettingsUpdatedId;
 	FCallbackBase* DisplaySettingsUpdatedCallback;
+
+	/** Last Login Credentials used for a login attempt */
+	TMap<int32, TSharedRef<FOnlineAccountCredentials>> LocalUserNumToLastLoginCredentials;
 };
 
 #endif
diff --git a/Engine/Plugins/Online/OnlineSubsystemEOS/Source/OnlineSubsystemEOSPlus/Private/OnlineAchievementsEOSPlus.cpp b/Engine/Plugins/Online/OnlineSubsystemEOS/Source/OnlineSubsystemEOSPlus/Private/OnlineAchievementsEOSPlus.cpp
--- a/Engine/Plugins/Online/OnlineSubsystemEOS/Source/OnlineSubsystemEOSPlus/Private/OnlineAchievementsEOSPlus.cpp
+++ b/Engine/Plugins/Online/OnlineSubsystemEOS/Source/OnlineSubsystemEOSPlus/Private/OnlineAchievementsEOSPlus.cpp
@@ -10,17 +10,22 @@ FOnAchievementsWrittenDelegate Ignored;
 FOnlineAchievementsEOSPlus::FOnlineAchievementsEOSPlus(FOnlineSubsystemEOSPlus* InSubsystem)
 	: EOSPlus(InSubsystem)
 {
-	BaseAchievementsInterface = EOSPlus->BaseOSS->GetAchievementsInterface();
-	check(BaseAchievementsInterface.IsValid());
 	EosAchievementsInterface = EOSPlus->EosOSS->GetAchievementsInterface();
 	check(EosAchievementsInterface.IsValid());
 
-	BaseAchievementsInterface->AddOnAchievementUnlockedDelegate_Handle(FOnAchievementUnlockedDelegate::CreateRaw(this, &FOnlineAchievementsEOSPlus::OnAchievementUnlocked));
+	BaseAchievementsInterface = EOSPlus->BaseOSS->GetAchievementsInterface();
+	if (BaseAchievementsInterface.IsValid())
+	{
+		BaseAchievementsInterface->AddOnAchievementUnlockedDelegate_Handle(FOnAchievementUnlockedDelegate::CreateRaw(this, &FOnlineAchievementsEOSPlus::OnAchievementUnlocked));
+	}
 }
 
 FOnlineAchievementsEOSPlus::~FOnlineAchievementsEOSPlus()
 {
-	BaseAchievementsInterface->ClearOnAchievementUnlockedDelegates(this);
+	if (BaseAchievementsInterface.IsValid())
+	{
+		BaseAchievementsInterface->ClearOnAchievementUnlockedDelegates(this);
+	}
 }
 
 void FOnlineAchievementsEOSPlus::OnAchievementUnlocked(const FUniqueNetId& PlayerId, const FString& AchievementId)
@@ -42,7 +47,7 @@ void FOnlineAchievementsEOSPlus::WriteAchievements(const FUniqueNetId& PlayerId,
 		return;
 	}
 
-	if (NetIdPlus->GetBaseNetId().IsValid())
+	if (NetIdPlus->GetBaseNetId().IsValid() && BaseAchievementsInterface.IsValid())
 	{
 		BaseAchievementsInterface->WriteAchievements(*NetIdPlus->GetBaseNetId(), WriteObject, Delegate);
 	}
@@ -65,7 +70,10 @@ void FOnlineAchievementsEOSPlus::QueryAchievements(const FUniqueNetId& PlayerId,
 		return;
 	}
 
-	BaseAchievementsInterface->QueryAchievements(*NetIdPlus->GetBaseNetId(), Delegate);
+	if (BaseAchievementsInterface.IsValid())
+	{
+		BaseAchievementsInterface->QueryAchievements(*NetIdPlus->GetBaseNetId(), Delegate);
+	}
 }
 
 void FOnlineAchievementsEOSPlus::QueryAchievementDescriptions(const FUniqueNetId& PlayerId, const FOnQueryAchievementsCompleteDelegate& Delegate)
@@ -77,7 +85,10 @@ void FOnlineAchievementsEOSPlus::QueryAchievementDescriptions(const FUniqueNetId
 		return;
 	}
 
-	BaseAchievementsInterface->QueryAchievementDescriptions(*NetIdPlus->GetBaseNetId(), Delegate);
+	if (BaseAchievementsInterface.IsValid())
+	{
+		BaseAchievementsInterface->QueryAchievementDescriptions(*NetIdPlus->GetBaseNetId(), Delegate);
+	}
 }
 
 EOnlineCachedResult::Type FOnlineAchievementsEOSPlus::GetCachedAchievement(const FUniqueNetId& PlayerId, const FString& AchievementId, FOnlineAchievement& OutAchievement)
@@ -89,7 +100,14 @@ EOnlineCachedResult::Type FOnlineAchievementsEOSPlus::GetCachedAchievement(const
 		return EOnlineCachedResult::NotFound;
 	}
 
-	return BaseAchievementsInterface->GetCachedAchievement(*NetIdPlus->GetBaseNetId(), AchievementId, OutAchievement);
+	EOnlineCachedResult::Type Result = EOnlineCachedResult::NotFound;
+
+	if (BaseAchievementsInterface.IsValid())
+	{
+		Result = BaseAchievementsInterface->GetCachedAchievement(*NetIdPlus->GetBaseNetId(), AchievementId, OutAchievement);
+	}
+
+	return Result;
 }
 
 EOnlineCachedResult::Type FOnlineAchievementsEOSPlus::GetCachedAchievements(const FUniqueNetId& PlayerId, TArray<FOnlineAchievement>& OutAchievements)
@@ -101,17 +119,38 @@ EOnlineCachedResult::Type FOnlineAchievementsEOSPlus::GetCachedAchievements(cons
 		return EOnlineCachedResult::NotFound;
 	}
 
-	return BaseAchievementsInterface->GetCachedAchievements(*NetIdPlus->GetBaseNetId(), OutAchievements);
+	EOnlineCachedResult::Type Result = EOnlineCachedResult::NotFound;
+
+	if (BaseAchievementsInterface.IsValid())
+	{
+		Result = BaseAchievementsInterface->GetCachedAchievements(*NetIdPlus->GetBaseNetId(), OutAchievements);
+	}
+
+	return Result;
 }
 
 EOnlineCachedResult::Type FOnlineAchievementsEOSPlus::GetCachedAchievementDescription(const FString& AchievementId, FOnlineAchievementDesc& OutAchievementDesc)
 {
-	return BaseAchievementsInterface->GetCachedAchievementDescription(AchievementId, OutAchievementDesc);
+	EOnlineCachedResult::Type Result = EOnlineCachedResult::NotFound;
+
+	if (BaseAchievementsInterface.IsValid())
+	{
+		Result = BaseAchievementsInterface->GetCachedAchievementDescription(AchievementId, OutAchievementDesc);
+	}
+
+	return Result;
 }
 
 #if !UE_BUILD_SHIPPING
 bool FOnlineAchievementsEOSPlus::ResetAchievements(const FUniqueNetId& PlayerId)
 {
-	return BaseAchievementsInterface->ResetAchievements(PlayerId);
+	bool bResult = false;
+
+	if (BaseAchievementsInterface.IsValid())
+	{
+		bResult = BaseAchievementsInterface->ResetAchievements(PlayerId);
+	}
+
+	return bResult;
 }
 #endif
diff --git a/Engine/Plugins/Online/OnlineSubsystemEOS/Source/OnlineSubsystemEOSPlus/Private/OnlineUserEOSPlus.cpp b/Engine/Plugins/Online/OnlineSubsystemEOS/Source/OnlineSubsystemEOSPlus/Private/OnlineUserEOSPlus.cpp
--- a/Engine/Plugins/Online/OnlineSubsystemEOS/Source/OnlineSubsystemEOSPlus/Private/OnlineUserEOSPlus.cpp
+++ b/Engine/Plugins/Online/OnlineSubsystemEOS/Source/OnlineSubsystemEOSPlus/Private/OnlineUserEOSPlus.cpp
@@ -75,7 +75,7 @@ inline FString BuildEOSPlusStringId(FUniqueNetIdPtr InBaseUniqueNetId, FUniqueNe
 
 FUniqueNetIdEOSPlus::FUniqueNetIdEOSPlus(FUniqueNetIdPtr InBaseUniqueNetId, FUniqueNetIdPtr InEOSUniqueNetId)
 	PRAGMA_DISABLE_DEPRECATION_WARNINGS
-	: FUniqueNetIdString(BuildEOSPlusStringId(InBaseUniqueNetId, InEOSUniqueNetId))
+	: FUniqueNetIdString(BuildEOSPlusStringId(InBaseUniqueNetId, InEOSUniqueNetId), FName("EOSPlus"))
 	PRAGMA_ENABLE_DEPRECATION_WARNINGS
 	, BaseUniqueNetId(InBaseUniqueNetId)
 	, EOSUniqueNetId(InEOSUniqueNetId)
@@ -163,6 +163,7 @@ FOnlineUserEOSPlus::FOnlineUserEOSPlus(FOnlineSubsystemEOSPlus* InSubsystem)
 	{
 		BaseIdentityInterface->AddOnLoginStatusChangedDelegate_Handle(LocalUserNum, FOnLoginStatusChangedDelegate::CreateRaw(this, &FOnlineUserEOSPlus::OnLoginStatusChanged));
 		EOSIdentityInterface->AddOnLoginCompleteDelegate_Handle(LocalUserNum, FOnLoginCompleteDelegate::CreateRaw(this, &FOnlineUserEOSPlus::OnLoginComplete));
+		BaseIdentityInterface->AddOnLoginCompleteDelegate_Handle(LocalUserNum, FOnLoginCompleteDelegate::CreateRaw(this, &FOnlineUserEOSPlus::OnBaseLoginComplete));
 		BaseIdentityInterface->AddOnLogoutCompleteDelegate_Handle(LocalUserNum, FOnLogoutCompleteDelegate::CreateRaw(this, &FOnlineUserEOSPlus::OnLogoutComplete));
 
 		BaseFriendsInterface->AddOnFriendsChangeDelegate_Handle(LocalUserNum, FOnFriendsChangeDelegate::CreateRaw(this, &FOnlineUserEOSPlus::OnFriendsChanged));
@@ -202,7 +203,7 @@ FOnlineUserEOSPlus::~FOnlineUserEOSPlus()
 	}
 }
 
-FUniqueNetIdEOSPlusPtr FOnlineUserEOSPlus::GetNetIdPlus(const FString& SourceId)
+FUniqueNetIdEOSPlusPtr FOnlineUserEOSPlus::GetNetIdPlus(const FString& SourceId) const
 {
 	if (NetIdPlusToNetIdPlus.Contains(SourceId))
 	{
@@ -219,7 +220,7 @@ FUniqueNetIdEOSPlusPtr FOnlineUserEOSPlus::GetNetIdPlus(const FString& SourceId)
 	return nullptr;
 }
 
-FUniqueNetIdPtr FOnlineUserEOSPlus::GetBaseNetId(const FString& SourceId)
+FUniqueNetIdPtr FOnlineUserEOSPlus::GetBaseNetId(const FString& SourceId) const
 {
 	if (NetIdPlusToBaseNetId.Contains(SourceId))
 	{
@@ -228,7 +229,7 @@ FUniqueNetIdPtr FOnlineUserEOSPlus::GetBaseNetId(const FString& SourceId)
 	return nullptr;
 }
 
-FUniqueNetIdPtr FOnlineUserEOSPlus::GetEOSNetId(const FString& SourceId)
+FUniqueNetIdPtr FOnlineUserEOSPlus::GetEOSNetId(const FString& SourceId) const
 {
 	if (NetIdPlusToEOSNetId.Contains(SourceId))
 	{
@@ -239,6 +240,8 @@ FUniqueNetIdPtr FOnlineUserEOSPlus::GetEOSNetId(const FString& SourceId)
 
 bool FOnlineUserEOSPlus::Login(int32 LocalUserNum, const FOnlineAccountCredentials& AccountCredentials)
 {
+	LocalUserNumToLastLoginCredentials.Emplace(LocalUserNum, MakeShared<FOnlineAccountCredentials>(AccountCredentials));
+
 	return BaseIdentityInterface->Login(LocalUserNum, AccountCredentials);
 }
 
@@ -323,6 +326,12 @@ void FOnlineUserEOSPlus::OnLoginComplete(int32 LocalUserNum, bool bWasSuccessful
 	TriggerOnLoginCompleteDelegates(LocalUserNum, bWasSuccessful, *NetIdPlus, Error);
 }
 
+void FOnlineUserEOSPlus::OnBaseLoginComplete(int32 LocalUserNum, bool bWasSuccessful, const FUniqueNetId& UserId, const FString& ErrorStr)
+{
+	check(LocalUserNumToLastLoginCredentials.Contains(LocalUserNum));
+	EOSIdentityInterface->Login(LocalUserNum, *LocalUserNumToLastLoginCredentials[LocalUserNum]);
+}
+
 void FOnlineUserEOSPlus::OnLogoutComplete(int32 LocalUserNum, bool bWasSuccessful)
 {
 	TriggerOnLogoutCompleteDelegates(LocalUserNum, bWasSuccessful);
@@ -395,12 +404,16 @@ bool FOnlineUserEOSPlus::Logout(int32 LocalUserNum)
 	// Clean up the cached data for this user
 	RemovePlayer(LocalUserNum);
 
+	LocalUserNumToLastLoginCredentials.Remove(LocalUserNum);
+
 	EOSIdentityInterface->Logout(LocalUserNum);
 	return BaseIdentityInterface->Logout(LocalUserNum);
 }
 
 bool FOnlineUserEOSPlus::AutoLogin(int32 LocalUserNum)
 {
+	LocalUserNumToLastLoginCredentials.Emplace(LocalUserNum, MakeShared<FOnlineAccountCredentials>(FOnlineAccountCredentials()));
+
 	return BaseIdentityInterface->AutoLogin(LocalUserNum);
 }
 
@@ -490,7 +503,14 @@ ELoginStatus::Type FOnlineUserEOSPlus::GetLoginStatus(int32 LocalUserNum) const
 
 ELoginStatus::Type FOnlineUserEOSPlus::GetLoginStatus(const FUniqueNetId& UserId) const
 {
-	return BaseIdentityInterface->GetLoginStatus(UserId);
+	TSharedPtr<FUniqueNetIdEOSPlus> NetIdPlus = GetNetIdPlus(UserId.ToString());
+	if (!NetIdPlus.IsValid())
+	{
+		UE_LOG_ONLINE(Error, TEXT("[FOnlineUserEOSPlus::GetLoginStatus] NetIdPlus not found for UserId %s"), *UserId.ToString());
+		return ELoginStatus::NotLoggedIn;
+	}
+
+	return BaseIdentityInterface->GetLoginStatus(*NetIdPlus->GetBaseNetId());
 }
 
 FString FOnlineUserEOSPlus::GetPlayerNickname(int32 LocalUserNum) const
@@ -511,7 +531,15 @@ FString FOnlineUserEOSPlus::GetAuthToken(int32 LocalUserNum) const
 
 void FOnlineUserEOSPlus::GetUserPrivilege(const FUniqueNetId& UserId, EUserPrivileges::Type Privilege, const FOnGetUserPrivilegeCompleteDelegate& Delegate)
 {
-	BaseIdentityInterface->GetUserPrivilege(UserId, Privilege, Delegate);
+	TSharedPtr<FUniqueNetIdEOSPlus> NetIdPlus = GetNetIdPlus(UserId.ToString());
+	if (NetIdPlus.IsValid())
+	{
+		BaseIdentityInterface->GetUserPrivilege(*NetIdPlus->GetBaseNetId(), Privilege, Delegate);
+	}
+	else
+	{
+		UE_LOG_ONLINE(Error, TEXT("[FOnlineUserEOSPlus::GetUserPrivilege] NetIdPlus not found for UserId %s"), *UserId.ToString());
+	}
 }
 
 FString FOnlineUserEOSPlus::GetAuthType() const
diff --git a/Engine/Plugins/Online/OnlineSubsystemEOS/Source/OnlineSubsystemEOSPlus/Private/OnlineUserEOSPlus.h b/Engine/Plugins/Online/OnlineSubsystemEOS/Source/OnlineSubsystemEOSPlus/Private/OnlineUserEOSPlus.h
--- a/Engine/Plugins/Online/OnlineSubsystemEOS/Source/OnlineSubsystemEOSPlus/Private/OnlineUserEOSPlus.h
+++ b/Engine/Plugins/Online/OnlineSubsystemEOS/Source/OnlineSubsystemEOSPlus/Private/OnlineUserEOSPlus.h
@@ -418,6 +418,7 @@ PACKAGE_SCOPE:
 	void OnLoginStatusChanged(int32 LocalUserNum, ELoginStatus::Type OldStatus, ELoginStatus::Type NewStatus, const FUniqueNetId& NewId);
 	void OnControllerPairingChanged(int32 LocalUserNum, FControllerPairingChangedUserInfo PreviousUser, FControllerPairingChangedUserInfo NewUser);
 	void OnLoginComplete(int32 LocalUserNum, bool bWasSuccessful, const FUniqueNetId& UserId, const FString& Error);
+	void OnBaseLoginComplete(int32 LocalUserNum, bool bWasSuccessful, const FUniqueNetId& UserId, const FString& ErrorStr);
 	void OnLogoutComplete(int32 LocalUserNum, bool bWasSuccessful);
 	void OnPresenceReceived(const FUniqueNetId& UserId, const TSharedRef<FOnlineUserPresence>& Presence);
 	void OnPresenceArrayUpdated(const FUniqueNetId& UserId, const TArray<TSharedRef<FOnlineUserPresence>>& NewPresenceArray);
@@ -430,9 +431,9 @@ PACKAGE_SCOPE:
 	void OnInviteAborted(const FUniqueNetId& UserId, const FUniqueNetId& FriendId);
 	void OnFriendRemoved(const FUniqueNetId& UserId, const FUniqueNetId& FriendId);
 
-	FUniqueNetIdEOSPlusPtr GetNetIdPlus(const FString& SourceId);
-	FUniqueNetIdPtr GetBaseNetId(const FString& SourceId);
-	FUniqueNetIdPtr GetEOSNetId(const FString& SourceId);
+	FUniqueNetIdEOSPlusPtr GetNetIdPlus(const FString& SourceId) const;
+	FUniqueNetIdPtr GetBaseNetId(const FString& SourceId) const;
+	FUniqueNetIdPtr GetEOSNetId(const FString& SourceId) const;
 
 private:
 	void AddPlayer(int32 LocalUserNum);
@@ -468,6 +469,9 @@ private:
 	TMap<FString, TSharedRef<FOnlineFriendPlus>> NetIdPlusToFriendMap;
 	TMap<FString, TSharedRef<FOnlineRecentPlayerPlus>> NetIdPlusToRecentPlayerMap;
 	TMap<FString, TSharedRef<FOnlineBlockedPlayerPlus>> NetIdPlusToBlockedPlayerMap;
+
+	/** Last Login Credentials used for a login attempt */
+	TMap<int32, TSharedRef<FOnlineAccountCredentials>> LocalUserNumToLastLoginCredentials;
 };
 
 typedef TSharedPtr<FOnlineUserEOSPlus, ESPMode::ThreadSafe> FOnlineUserEOSPlusPtr;
