Variable refactor to fix Horde Issue #73729

diff --git a/Engine/Plugins/Online/OnlineSubsystemEOS/Source/OnlineSubsystemEOSPlus/Private/OnlineUserEOSPlus.cpp b/Engine/Plugins/Online/OnlineSubsystemEOS/Source/OnlineSubsystemEOSPlus/Private/OnlineUserEOSPlus.cpp
--- a/Engine/Plugins/Online/OnlineSubsystemEOS/Source/OnlineSubsystemEOSPlus/Private/OnlineUserEOSPlus.cpp
+++ b/Engine/Plugins/Online/OnlineSubsystemEOS/Source/OnlineSubsystemEOSPlus/Private/OnlineUserEOSPlus.cpp
@@ -503,7 +503,7 @@ ELoginStatus::Type FOnlineUserEOSPlus::GetLoginStatus(int32 LocalUserNum) const
 
 ELoginStatus::Type FOnlineUserEOSPlus::GetLoginStatus(const FUniqueNetId& UserId) const
 {
-	TSharedPtr<FUniqueNetIdEOSPlus> NetIdPlus = GetNetIdPlus(UserId.ToString());
+	FUniqueNetIdEOSPlusPtr NetIdPlus = GetNetIdPlus(UserId.ToString());
 	if (!NetIdPlus.IsValid())
 	{
 		UE_LOG_ONLINE(Error, TEXT("[FOnlineUserEOSPlus::GetLoginStatus] NetIdPlus not found for UserId %s"), *UserId.ToString());
@@ -531,7 +531,7 @@ FString FOnlineUserEOSPlus::GetAuthToken(int32 LocalUserNum) const
 
 void FOnlineUserEOSPlus::GetUserPrivilege(const FUniqueNetId& UserId, EUserPrivileges::Type Privilege, const FOnGetUserPrivilegeCompleteDelegate& Delegate)
 {
-	TSharedPtr<FUniqueNetIdEOSPlus> NetIdPlus = GetNetIdPlus(UserId.ToString());
+	FUniqueNetIdEOSPlusPtr NetIdPlus = GetNetIdPlus(UserId.ToString());
 	if (NetIdPlus.IsValid())
 	{
 		BaseIdentityInterface->GetUserPrivilege(*NetIdPlus->GetBaseNetId(), Privilege, Delegate);
