Added missing delegate triggers from EOS Session and Lobby methods.

diff --git a/Engine/Plugins/Online/OnlineSubsystemEOS/Source/OnlineSubsystemEOS/Private/OnlineSessionEOS.cpp b/Engine/Plugins/Online/OnlineSubsystemEOS/Source/OnlineSubsystemEOS/Private/OnlineSessionEOS.cpp
--- a/Engine/Plugins/Online/OnlineSubsystemEOS/Source/OnlineSubsystemEOS/Private/OnlineSessionEOS.cpp
+++ b/Engine/Plugins/Online/OnlineSubsystemEOS/Source/OnlineSubsystemEOS/Private/OnlineSessionEOS.cpp
@@ -461,6 +461,7 @@ void FOnlineSessionEOS::Init(const FString& InBucketId)
 		if (!NetId.IsValid())
 		{
 			UE_LOG_ONLINE_SESSION(Warning, TEXT("Cannot accept invite due to unknown user (%s)"), *LexToString(Data->LocalUserId));
+			TriggerOnSessionUserInviteAcceptedDelegates(false, 0, NetId, FOnlineSessionSearchResult());
 			return;
 		}
 		int32 LocalUserNum = EOSSubsystem->UserManager->GetLocalUserNumFromUniqueNetId(*NetId);
@@ -476,6 +477,11 @@ void FOnlineSessionEOS::Init(const FString& InBucketId)
 			AddSearchResult(SessionDetails, LastInviteSearch.ToSharedRef());
 			TriggerOnSessionUserInviteAcceptedDelegates(true, LocalUserNum, NetId, LastInviteSearch->SearchResults[0]);
 		}
+		else
+		{
+			UE_LOG_ONLINE_SESSION(Warning, TEXT("EOS_Sessions_CopySessionHandleByInviteId not successful. Finished with EOS_EResult %s"), ANSI_TO_TCHAR(EOS_EResult_ToString(Result)));
+			TriggerOnSessionUserInviteAcceptedDelegates(false, LocalUserNum, NetId, FOnlineSessionSearchResult());
+		}
 	};
 	EOS_Sessions_AddNotifySessionInviteAcceptedOptions Options = { };
 	Options.ApiVersion = EOS_SESSIONS_ADDNOTIFYSESSIONINVITEACCEPTED_API_LATEST;
@@ -813,6 +819,7 @@ void FOnlineSessionEOS::OnLobbyInviteAccepted(const char* InviteId, const EOS_Pr
 	if (!NetId.IsValid())
 	{
 		UE_LOG_ONLINE_SESSION(Warning, TEXT("[FOnlineSessionEOS::OnLobbyInviteAccepted] Cannot accept lobby invite due to unknown user (%s)"), *LexToString(LocalUserId));
+		TriggerOnSessionUserInviteAcceptedDelegates(false, 0, NetId, FOnlineSessionSearchResult());
 		return;
 	}
 	int32 LocalUserNum = EOSSubsystem->UserManager->GetLocalUserNumFromUniqueNetId(*NetId);
@@ -833,6 +840,7 @@ void FOnlineSessionEOS::OnLobbyInviteAccepted(const char* InviteId, const EOS_Pr
 	else
 	{
 		UE_LOG_ONLINE_SESSION(Warning, TEXT("[FOnlineSessionEOS::OnLobbyInviteAccepted] EOS_Lobby_CopyLobbyDetailsHandleByInviteId failed with EOS result code (%s)"), ANSI_TO_TCHAR(EOS_EResult_ToString(Result)));
+		TriggerOnSessionUserInviteAcceptedDelegates(false, LocalUserNum, NetId, FOnlineSessionSearchResult());
 	}
 }
 
@@ -842,6 +850,7 @@ void FOnlineSessionEOS::OnJoinLobbyAccepted(const EOS_ProductUserId& LocalUserId
 	if (!NetId.IsValid())
 	{
 		UE_LOG_ONLINE_SESSION(Warning, TEXT("[FOnlineSessionEOS::OnJoinLobbyAccepted] Cannot join lobby due to unknown user (%s)"), *LexToString(LocalUserId));
+		TriggerOnSessionUserInviteAcceptedDelegates(false, 0, NetId, FOnlineSessionSearchResult());
 		return;
 	}
 	int32 LocalUserNum = EOSSubsystem->UserManager->GetLocalUserNumFromUniqueNetId(*NetId);
@@ -861,6 +870,7 @@ void FOnlineSessionEOS::OnJoinLobbyAccepted(const EOS_ProductUserId& LocalUserId
 	else
 	{
 		UE_LOG_ONLINE_SESSION(Warning, TEXT("[FOnlineSessionEOS::OnJoinLobbyAccepted] EOS_Lobby_CopyLobbyDetailsHandleByUiEventId failed with EOS result code (%s)"), ANSI_TO_TCHAR(EOS_EResult_ToString(Result)));
+		TriggerOnSessionUserInviteAcceptedDelegates(false, LocalUserNum, NetId, FOnlineSessionSearchResult());
 	}
 }
 
@@ -2451,6 +2461,10 @@ bool FOnlineSessionEOS::FindFriendSession(int32 LocalUserNum, const FUniqueNetId
 	else
 	{
 		UE_LOG_ONLINE_SESSION(Warning, TEXT("[FOnlineSessionEOS::FindFriendSession] CreateLobbySearch not successful. Finished with EOS_EResult %s"), ANSI_TO_TCHAR(EOS_EResult_ToString(CreateLobbySearchResult)));
+		EOSSubsystem->ExecuteNextTick([this]()
+		{
+			TriggerOnFindSessionsCompleteDelegates(false);
+		});
 	}
 
 	return bResult;
@@ -3440,6 +3454,11 @@ uint32 FOnlineSessionEOS::JoinLobbySession(int32 PlayerNum, FNamedOnlineSession*
 
 					LobbySearchResultsCache.Reset();
 				}
+				else
+				{
+					UE_LOG_ONLINE_SESSION(Warning, TEXT("[FOnlineSessionEOS::JoinLobbySession] Unable to find session %s"), *SessionName.ToString());
+					TriggerOnJoinSessionCompleteDelegates(SessionName, EOnJoinSessionCompleteResult::SessionDoesNotExist);
+				}
 			};
 
 			EOS_Lobby_JoinLobby(LobbyHandle, &JoinLobbyOptions, CallbackObj, CallbackObj->GetCallbackPtr());
@@ -3639,7 +3658,12 @@ uint32 FOnlineSessionEOS::UpdateLobbySession(FNamedOnlineSession* Session)
 						UE_LOG_ONLINE_SESSION(Warning, TEXT("[FOnlineSessionEOS::UpdateLobbySession] UpdateLobby not successful. Finished with EOS_EResult %s"), ANSI_TO_TCHAR(EOS_EResult_ToString(Data->ResultCode)));
 					}
 
-					TriggerOnUpdateSessionCompleteDelegates(Session->SessionName, bWasSuccessful);
+					TriggerOnUpdateSessionCompleteDelegates(SessionName, bWasSuccessful);
+				}
+				else
+				{
+					UE_LOG_ONLINE_SESSION(Warning, TEXT("[FOnlineSessionEOS::UpdateLobbySession] Unable to find session %s"), *SessionName.ToString());
+					TriggerOnUpdateSessionCompleteDelegates(SessionName, false);
 				}
 			};
 
@@ -3729,6 +3753,11 @@ uint32 FOnlineSessionEOS::DestroyLobbySession(FNamedOnlineSession* Session, cons
 					CompletionDelegate.ExecuteIfBound(SessionName, bWasSuccessful);
 					TriggerOnDestroySessionCompleteDelegates(SessionName, bWasSuccessful);
 				}
+				else
+				{
+					UE_LOG_ONLINE_SESSION(Warning, TEXT("[FOnlineSessionEOS::DestroyLobbySession] Unable to find session %s"), *SessionName.ToString());
+					TriggerOnDestroySessionCompleteDelegates(SessionName, false);
+				}
 			};
 
 			EOS_Lobby_DestroyLobby(LobbyHandle, &DestroyOptions, DestroyCallbackObj, DestroyCallbackObj->GetCallbackPtr());
@@ -3775,6 +3804,11 @@ uint32 FOnlineSessionEOS::DestroyLobbySession(FNamedOnlineSession* Session, cons
 					CompletionDelegate.ExecuteIfBound(SessionName, bWasSuccessful);
 					TriggerOnDestroySessionCompleteDelegates(SessionName, bWasSuccessful);
 				}
+				else
+				{
+					UE_LOG_ONLINE_SESSION(Warning, TEXT("[FOnlineSessionEOS::DestroyLobbySession] Unable to find session %s"), *SessionName.ToString());
+					TriggerOnDestroySessionCompleteDelegates(SessionName, false);
+				}
 			};
 
 			EOS_Lobby_LeaveLobby(LobbyHandle, &LeaveOptions, LeaveCallbackObj, LeaveCallbackObj->GetCallbackPtr());
