Changed EOSPlus native OSS log from warning to verbose.

diff --git a/Engine/Plugins/Online/OnlineSubsystemEOS/Source/OnlineSubsystemEOSPlus/Private/OnlineSubsystemModuleEOSPlus.cpp b/Engine/Plugins/Online/OnlineSubsystemEOS/Source/OnlineSubsystemEOSPlus/Private/OnlineSubsystemModuleEOSPlus.cpp
--- a/Engine/Plugins/Online/OnlineSubsystemEOS/Source/OnlineSubsystemEOSPlus/Private/OnlineSubsystemModuleEOSPlus.cpp
+++ b/Engine/Plugins/Online/OnlineSubsystemEOS/Source/OnlineSubsystemEOSPlus/Private/OnlineSubsystemModuleEOSPlus.cpp
@@ -77,7 +77,7 @@ void FOnlineSubsystemEOSPlusModule::LoadBaseOSS()
 	}
 	if (BaseOSSName.IsEmpty())
 	{
-		UE_LOG_ONLINE(Warning, TEXT("EOSPlus failed to find the native OSS!"));
+		UE_LOG_ONLINE(Verbose, TEXT("EOSPlus failed to find the native OSS!"));
 		return;
 	}
 	// Load the module for the base OSS
