Forcing EOS Stat names to uppercase before calling API methods.

diff --git a/Engine/Plugins/Online/OnlineSubsystemEOS/Source/OnlineSubsystemEOS/Private/OnlineStatsEOS.cpp b/Engine/Plugins/Online/OnlineSubsystemEOS/Source/OnlineSubsystemEOS/Private/OnlineStatsEOS.cpp
--- a/Engine/Plugins/Online/OnlineSubsystemEOS/Source/OnlineSubsystemEOS/Private/OnlineStatsEOS.cpp
+++ b/Engine/Plugins/Online/OnlineSubsystemEOS/Source/OnlineSubsystemEOS/Private/OnlineStatsEOS.cpp
@@ -125,7 +125,7 @@ void FOnlineStatsEOS::QueryStats(const FUniqueNetIdRef LocalUserId, const TArray
 	FQueryStatsOptions Options(StatNames.Num());
 	for (int32 Index = 0; Index < StatNames.Num(); Index++)
 	{
-		FCStringAnsi::Strncpy(Options.PointerArray[Index], TCHAR_TO_UTF8(*StatNames[Index]), EOS_OSS_STRING_BUFFER_LENGTH);
+		FCStringAnsi::Strncpy(Options.PointerArray[Index], TCHAR_TO_UTF8(*StatNames[Index].ToUpper()), EOS_OSS_STRING_BUFFER_LENGTH);
 	}
 
 	int32 NumKnownUsers = 0;
@@ -176,7 +176,7 @@ void FOnlineStatsEOS::QueryStats(const FUniqueNetIdRef LocalUserId, const TArray
 					// Read each stat that we were looking for so we can mark missing ones as "empty"
 					for (const FString& StatName : StatsQueryContext->StatNames)
 					{
-						FCStringAnsi::Strncpy(StatNameANSI, TCHAR_TO_UTF8(*StatName), EOS_OSS_STRING_BUFFER_LENGTH);
+						FCStringAnsi::Strncpy(StatNameANSI, TCHAR_TO_UTF8(*StatName.ToUpper()), EOS_OSS_STRING_BUFFER_LENGTH);
 
 						EOS_Stats_Stat* ReadStat = nullptr;
 						if (EOS_Stats_CopyStatByName(EOSSubsystem->StatsHandle, &Options, &ReadStat) == EOS_EResult::EOS_Success)
@@ -296,7 +296,7 @@ void FOnlineStatsEOS::WriteStats(EOS_ProductUserId LocalUserId, EOS_ProductUserI
 		EOSStat.ApiVersion = EOS_STATS_INGESTDATA_API_LATEST;
 
 		EOSStat.IngestAmount = GetVariantValue(Stat.Value.GetValue());
-		FCStringAnsi::Strncpy(EOSStatNames[Index].StatName, TCHAR_TO_UTF8(*Stat.Key), EOS_OSS_STRING_BUFFER_LENGTH);
+		FCStringAnsi::Strncpy(EOSStatNames[Index].StatName, TCHAR_TO_UTF8(*Stat.Key.ToUpper()), EOS_OSS_STRING_BUFFER_LENGTH);
 		EOSStat.StatName = EOSStatNames[Index].StatName;
 
 		Index++;
