Stop querying EOS presence of players that aren't our friend. Causes an SDK error. EOS_Friends_QueryFriends returns both our friends and players we have only sent or received friend invites from.

diff --git a/Engine/Plugins/Online/OnlineSubsystemEOS/Source/OnlineSubsystemEOS/Private/UserManagerEOS.cpp b/Engine/Plugins/Online/OnlineSubsystemEOS/Source/OnlineSubsystemEOS/Private/UserManagerEOS.cpp
--- a/Engine/Plugins/Online/OnlineSubsystemEOS/Source/OnlineSubsystemEOS/Private/UserManagerEOS.cpp
+++ b/Engine/Plugins/Online/OnlineSubsystemEOS/Source/OnlineSubsystemEOS/Private/UserManagerEOS.cpp
@@ -1662,8 +1662,16 @@ void FUserManagerEOS::AddFriend(int32 LocalUserNum, EOS_EpicAccountId EpicAccoun
 	
 	FriendRef->SetInviteStatus(ToEInviteStatus(Status));
 
-	// Add this friend as a remote (this will grab presence & user info)
+	// Add this friend as a remote player (this will grab user info)
 	AddRemotePlayer(LocalUserNum, NetId, EpicAccountId, FriendNetId, FriendRef, FriendRef);
+
+	// Querying the presence of a non-friend would cause an SDK error.
+	// Players that sent/recieved a friend invitation from us still count as "friends", so check
+	// our friend relationship here.
+	if(Status == EOS_EFriendsStatus::EOS_FS_Friends)
+	{
+		QueryPresence(*FriendNetId, IgnoredPresenceDelegate);
+	}
 }
 
 void FUserManagerEOS::AddRemotePlayer(int32 LocalUserNum, const FString& NetId, EOS_EpicAccountId EpicAccountId)
@@ -1686,8 +1694,6 @@ void FUserManagerEOS::AddRemotePlayer(int32 LocalUserNum, const FString& NetId,
 
 	// Read the user info for this player
 	ReadUserInfo(LocalUserNum, EpicAccountId);
-	// Read presence for this remote player
-	QueryPresence(*UniqueNetId, IgnoredPresenceDelegate);
 	// Get their product id mapping
 	FUniqueNetIdEOSPtr LocalNetId = GetLocalUniqueNetIdEOS(DefaultLocalUser);
 	if (LocalNetId.IsValid())
