[EOS OSS] Batch calls to EOS_Connect_QueryExternalAccountMappings as is already supported by the SDK.
Don't execute OnFriendsChange delegates when reading the friends list until the external mappings query has finished. Triggering it beforehand can cause friend's net ids to change after game code received them. If game code hashed those ids (ie. as the social toolkit does) the hashes become invalid.
Fixup QueryExternalIdMappings to correctly report the success for failure of the function to it's callback. Formerly it would always report a failure.

diff --git a/Engine/Plugins/Online/OnlineSubsystemEOS/Source/OnlineSubsystemEOS/Private/OnlineSubsystemEOSTypes.h b/Engine/Plugins/Online/OnlineSubsystemEOS/Source/OnlineSubsystemEOS/Private/OnlineSubsystemEOSTypes.h
--- a/Engine/Plugins/Online/OnlineSubsystemEOS/Source/OnlineSubsystemEOS/Private/OnlineSubsystemEOSTypes.h
+++ b/Engine/Plugins/Online/OnlineSubsystemEOS/Source/OnlineSubsystemEOS/Private/OnlineSubsystemEOSTypes.h
@@ -505,6 +505,27 @@ private:
 	}
 };
 
+namespace OSSInternalCallback
+{
+	/** Create a callback for a non-SDK function that is tied to the lifetime of an arbitrary shared pointer. */
+	template <typename DelegateType, typename OwnerType, typename... CallbackArgs>
+	UE_NODISCARD DelegateType Create(const TSharedPtr<OwnerType, ESPMode::ThreadSafe>& InOwner,
+		const TFunction<void(CallbackArgs...)>& InUserCallback)
+	{
+		const DelegateType& CheckOwnerThenExecute = DelegateType::CreateLambda(
+			[WeakOwner = TWeakPtr<OwnerType, ESPMode::ThreadSafe>(InOwner), InUserCallback](CallbackArgs... Payload) {
+				check(IsInGameThread());
+				TSharedPtr<OwnerType, ESPMode::ThreadSafe> Owner = WeakOwner.Pin();
+				if (Owner.IsValid())
+				{
+					InUserCallback(Payload...);
+				}
+		});
+
+		return CheckOwnerThenExecute;
+	}
+}
+
 /**
  * Class to handle nested callbacks (callbacks that are tied to an external callback's lifetime,
  * e.g. file chunkers) generically using a lambda to process callback results
diff --git a/Engine/Plugins/Online/OnlineSubsystemEOS/Source/OnlineSubsystemEOS/Private/UserManagerEOS.cpp b/Engine/Plugins/Online/OnlineSubsystemEOS/Source/OnlineSubsystemEOS/Private/UserManagerEOS.cpp
--- a/Engine/Plugins/Online/OnlineSubsystemEOS/Source/OnlineSubsystemEOS/Private/UserManagerEOS.cpp
+++ b/Engine/Plugins/Online/OnlineSubsystemEOS/Source/OnlineSubsystemEOS/Private/UserManagerEOS.cpp
@@ -1032,7 +1032,18 @@ bool FUserManagerEOS::IsFriendQueryUserInfoOngoing(int32 LocalUserNum)
 	// If we have an entry for this user and the corresponding array has any element, users are still being processed
 	if (IsFriendQueryUserInfoOngoingForLocalUserMap.Contains(LocalUserNum))
 	{
-		return IsFriendQueryUserInfoOngoingForLocalUserMap[LocalUserNum].Num() > 0;
+		if(IsFriendQueryUserInfoOngoingForLocalUserMap[LocalUserNum].Num() > 0)
+		{
+			return true;
+		}
+	}
+	
+	if (IsPlayerQueryExternalMappingsOngoingForLocalUserMap.Contains(LocalUserNum))
+	{
+		if(IsPlayerQueryExternalMappingsOngoingForLocalUserMap[LocalUserNum].Num() > 0)
+		{
+			return true;
+		}
 	}
 
 	return false;
@@ -1694,14 +1705,6 @@ void FUserManagerEOS::AddRemotePlayer(int32 LocalUserNum, const FString& NetId,
 
 	// Read the user info for this player
 	ReadUserInfo(LocalUserNum, EpicAccountId);
-	// Get their product id mapping
-	FUniqueNetIdEOSPtr LocalNetId = GetLocalUniqueNetIdEOS(DefaultLocalUser);
-	if (LocalNetId.IsValid())
-	{
-		TArray<FString> ExternalIds;
-		ExternalIds.Add(NetId);
-		QueryExternalIdMappings(*LocalNetId, FExternalIdQueryOptions(), ExternalIds, IgnoredMappingDelegate);
-	}
 }
 
 void FUserManagerEOS::UpdateRemotePlayerProductUserId(EOS_EpicAccountId AccountId, EOS_ProductUserId UserId)
@@ -1806,7 +1809,6 @@ bool FUserManagerEOS::ReadFriendsList(int32 LocalUserNum, const FString& ListNam
 			Result = EOS_EResult::EOS_InvalidUser;
 		}
 
-		FString ErrorString;
 		bool bWasSuccessful = Result == EOS_EResult::EOS_Success;
 		if (bWasSuccessful)
 		{
@@ -1817,6 +1819,8 @@ bool FUserManagerEOS::ReadFriendsList(int32 LocalUserNum, const FString& ListNam
 
 			LocalUserNumToFriendsListMap[LocalUserNum]->Empty(FriendCount);
 
+			TArray<FString> FriendEasIds;
+			FriendEasIds.Reserve(FriendCount);
 			// Process each friend returned
 			for (int32 Index = 0; Index < FriendCount; Index++)
 			{
@@ -1828,16 +1832,26 @@ bool FUserManagerEOS::ReadFriendsList(int32 LocalUserNum, const FString& ListNam
 				if (FriendEpicAccountId != nullptr)
 				{
 					AddFriend(LocalUserNum, FriendEpicAccountId);
+					FriendEasIds.Add(LexToString(FriendEpicAccountId));
 				}
 			}
+
+			const TFunction<FOnQueryExternalIdMappingsComplete::TFuncType>& OnExternalIdMappingsQueriedLambda = 
+				[this, LocalUserNum](bool bWasSuccessful, const FUniqueNetId& UserId, const FExternalIdQueryOptions& QueryOptions,
+					const TArray<FString>& ExternalIds, const FString& Error) {
+						ProcessReadFriendsListComplete(LocalUserNum, bWasSuccessful, Error);
+				};
+
+			const auto& ExternalMappingsCallback = OSSInternalCallback::Create<FOnQueryExternalIdMappingsComplete>(EOSSubsystem->UserManager, OnExternalIdMappingsQueriedLambda);
+			
+			QueryExternalIdMappings(*GetLocalUniqueNetIdEOS(DefaultLocalUser), FExternalIdQueryOptions(), FriendEasIds, ExternalMappingsCallback);
+			
 		}
 		else
 		{
-			ErrorString = FString::Printf(TEXT("ReadFriendsList(%d) failed with EOS result code (%s)"), LocalUserNum, ANSI_TO_TCHAR(EOS_EResult_ToString(Result)));
+			const FString ErrorString = FString::Printf(TEXT("ReadFriendsList(%d) failed with EOS result code (%s)"), LocalUserNum, ANSI_TO_TCHAR(EOS_EResult_ToString(Result)));
+			ProcessReadFriendsListComplete(LocalUserNum, false, ErrorString);
 		}
-
-		// Then evaluate the current state of the query to see what to do next
-		ProcessReadFriendsListComplete(LocalUserNum, bWasSuccessful, ErrorString);
 	};
 	EOS_Friends_QueryFriends(EOSSubsystem->FriendsHandle, &Options, CallbackObj, CallbackObj->GetCallbackPtr());
 
@@ -1862,6 +1876,7 @@ void FUserManagerEOS::ProcessReadFriendsListComplete(int32 LocalUserNum, bool bW
 		TriggerOnFriendsChangeDelegates(LocalUserNum);
 
 		IsFriendQueryUserInfoOngoingForLocalUserMap.Remove(LocalUserNum);
+		IsPlayerQueryExternalMappingsOngoingForLocalUserMap.Remove(LocalUserNum);
 	}
 }
 
@@ -2516,6 +2531,9 @@ EOnlineCachedResult::Type FUserManagerEOS::GetCachedPresenceForApp(const FUnique
 
 bool FUserManagerEOS::QueryUserInfo(int32 LocalUserNum, const TArray<FUniqueNetIdRef>& UserIds)
 {
+	TArray<FString> UserEasIdsNeedingExternalMappings;
+	UserEasIdsNeedingExternalMappings.Reserve(UserIds.Num());
+	
 	// Trigger a query for each user in the list
 	for (const FUniqueNetIdRef& NetId : UserIds)
 	{
@@ -2537,11 +2555,17 @@ bool FUserManagerEOS::QueryUserInfo(int32 LocalUserNum, const TArray<FUniqueNetI
 			EOS_EpicAccountId AccountId = EOS_EpicAccountId_FromString(TCHAR_TO_UTF8(*EOSID.EpicAccountIdStr));
 			if (EOS_EpicAccountId_IsValid(AccountId) == EOS_TRUE)
 			{
+				UserEasIdsNeedingExternalMappings.Add(EOSID.EpicAccountIdStr);
+				
 				// Registering the player will also query the user info data
 				AddRemotePlayer(LocalUserNum, EOSID.UniqueNetIdStr, AccountId);
 			}
 		}
 	}
+
+	const FUniqueNetIdEOSPtr LocalId = GetLocalUniqueNetIdEOS(LocalUserNum);
+	QueryExternalIdMappings(*LocalId, FExternalIdQueryOptions(), UserEasIdsNeedingExternalMappings, IgnoredMappingDelegate);
+	
 	return true;
 }
 
@@ -2559,7 +2583,7 @@ void FUserManagerEOS::ReadUserInfo(int32 LocalUserNum, EOS_EpicAccountId EpicAcc
 		}
 
 		// We mark this player as processed
-		IsFriendQueryUserInfoOngoingForLocalUserMap[LocalUserNum].Remove(EpicAccountId);
+		IsFriendQueryUserInfoOngoingForLocalUserMap[LocalUserNum].RemoveSwap(EpicAccountId, false);
 
 		ProcessReadFriendsListComplete(LocalUserNum, true, TEXT(""));
 	};
@@ -2722,6 +2746,9 @@ bool FUserManagerEOS::QueryExternalIdMappings(const FUniqueNetId& UserId, const
 	}
 	int32 LocalUserNum = GetLocalUserNumFromUniqueNetId(UserId);
 
+	// Mark the queries as in progress
+	IsPlayerQueryExternalMappingsOngoingForLocalUserMap.FindOrAdd(LocalUserNum).Append(ExternalIds);
+
 	EOS_ProductUserId LocalUserId = StringToProductUserIdMap[EOSID.UniqueNetIdStr];
 	const int32 NumBatches = (ExternalIds.Num() / EOS_CONNECT_QUERYEXTERNALACCOUNTMAPPINGS_MAX_ACCOUNT_IDS) + 1;
 	int32 QueryStart = 0;
@@ -2772,7 +2799,16 @@ bool FUserManagerEOS::QueryExternalIdMappings(const FUniqueNetId& UserId, const
 			{
 				ErrorString = FString::Printf(TEXT("EOS_Connect_QueryExternalAccountMappings() failed with result code (%s)"), ANSI_TO_TCHAR(EOS_EResult_ToString(Result)));
 			}
-			Delegate.ExecuteIfBound(false, *EOSID, QueryOptions, BatchIds, ErrorString);
+
+			// Mark all queries as complete
+			TArray<FString>& OngoingQueries = IsPlayerQueryExternalMappingsOngoingForLocalUserMap[LocalUserNum];
+			for (const FString& StringId : BatchIds)
+			{
+				OngoingQueries.RemoveSwap(StringId, false);
+			}
+			
+			const bool bWasSuccessful = Result == EOS_EResult::EOS_Success;
+			Delegate.ExecuteIfBound(bWasSuccessful, *EOSID, QueryOptions, BatchIds, ErrorString);
 		};
 
 		EOS_Connect_QueryExternalAccountMappings(EOSSubsystem->ConnectHandle, &Options, CallbackObj, CallbackObj->GetCallbackPtr());
diff --git a/Engine/Plugins/Online/OnlineSubsystemEOS/Source/OnlineSubsystemEOS/Private/UserManagerEOS.h b/Engine/Plugins/Online/OnlineSubsystemEOS/Source/OnlineSubsystemEOS/Private/UserManagerEOS.h
--- a/Engine/Plugins/Online/OnlineSubsystemEOS/Source/OnlineSubsystemEOS/Private/UserManagerEOS.h
+++ b/Engine/Plugins/Online/OnlineSubsystemEOS/Source/OnlineSubsystemEOS/Private/UserManagerEOS.h
@@ -449,6 +449,8 @@ private:
 
 	/** Id map to keep track of which friends have been processed during async user info queries */
 	TMap<int32, TArray<EOS_EpicAccountId>> IsFriendQueryUserInfoOngoingForLocalUserMap;
+	/** Id map to keep track of which players still need their external id synced */
+	TMap<int32, TArray<FString>> IsPlayerQueryExternalMappingsOngoingForLocalUserMap;
 
 	/** Cache for the info passed on to ReadFriendsList, kept while user info queries complete */
 	struct ReadUserListInfo
