- Fix crash due to calling SDK function when SDK is not loaded

diff --git a/Engine/Plugins/Online/OnlineSubsystemEOS/Source/OnlineSubsystemEOS/Private/OnlineSubsystemEOS.cpp b/Engine/Plugins/Online/OnlineSubsystemEOS/Source/OnlineSubsystemEOS/Private/OnlineSubsystemEOS.cpp
--- a/Engine/Plugins/Online/OnlineSubsystemEOS/Source/OnlineSubsystemEOS/Private/OnlineSubsystemEOS.cpp
+++ b/Engine/Plugins/Online/OnlineSubsystemEOS/Source/OnlineSubsystemEOS/Private/OnlineSubsystemEOS.cpp
@@ -173,7 +173,7 @@ void FOnlineSubsystemEOS::ModuleInit()
 	EOS_EResult InitResult = EOSSDKManager->Initialize();
 	if (InitResult != EOS_EResult::EOS_Success)
 	{
-		UE_LOG_ONLINE(Error, TEXT("FOnlineSubsystemEOS: failed to initialize the EOS SDK with result code (%s)"), ANSI_TO_TCHAR(EOS_EResult_ToString(InitResult)));
+		UE_LOG_ONLINE(Error, TEXT("FOnlineSubsystemEOS: failed to initialize the EOS SDK with result code (%d)"), InitResult);
 		return;
 	}
 }
