Added missing session name parameter from EOS Player Registration methods.
Added missing presence search attribute in EOS Sessions.

diff --git a/Engine/Plugins/Online/OnlineSubsystemEOS/Source/OnlineSubsystemEOS/Private/OnlineSessionEOS.cpp b/Engine/Plugins/Online/OnlineSubsystemEOS/Source/OnlineSubsystemEOS/Private/OnlineSessionEOS.cpp
--- a/Engine/Plugins/Online/OnlineSubsystemEOS/Source/OnlineSubsystemEOS/Private/OnlineSessionEOS.cpp
+++ b/Engine/Plugins/Online/OnlineSubsystemEOS/Source/OnlineSubsystemEOS/Private/OnlineSessionEOS.cpp
@@ -1054,6 +1054,11 @@ void FOnlineSessionEOS::AddAttribute(EOS_HSessionModification SessionModHandle,
 
 void FOnlineSessionEOS::SetAttributes(EOS_HSessionModification SessionModHandle, FNamedOnlineSession* Session)
 {
+	// The first will let us find it on session searches
+	const FString SearchPresence(SEARCH_PRESENCE.ToString());
+	const FAttributeOptions SearchPresenceAttribute(TCHAR_TO_UTF8(*SearchPresence), true);
+	AddAttribute(SessionModHandle, &SearchPresenceAttribute);
+
 	FAttributeOptions Opt1("NumPrivateConnections", Session->SessionSettings.NumPrivateConnections);
 	AddAttribute(SessionModHandle, &Opt1);
 
@@ -2794,6 +2799,8 @@ bool FOnlineSessionEOS::RegisterPlayers(FName SessionName, const TArray< FUnique
 			Options.ApiVersion = EOS_SESSIONS_REGISTERPLAYERS_API_LATEST;
 			Options.PlayersToRegister = EOSIds.GetData();
 			Options.PlayersToRegisterCount = EOSIds.Num();
+			const FTCHARToUTF8 Utf8SessionName(*SessionName.ToString());
+			Options.SessionName = Utf8SessionName.Get();
 
 			FRegisterPlayersCallback* CallbackObj = new FRegisterPlayersCallback();
 			CallbackObj->CallbackLambda = [this, SessionName, RegisteredPlayers = TArray<FUniqueNetIdRef>(Players)](const EOS_Sessions_RegisterPlayersCallbackInfo* Data)
@@ -2876,6 +2883,8 @@ bool FOnlineSessionEOS::UnregisterPlayers(FName SessionName, const TArray< FUniq
 			Options.ApiVersion = EOS_SESSIONS_UNREGISTERPLAYERS_API_LATEST;
 			Options.PlayersToUnregister = EOSIds.GetData();
 			Options.PlayersToUnregisterCount = EOSIds.Num();
+			const FTCHARToUTF8 Utf8SessionName(*SessionName.ToString());
+			Options.SessionName = Utf8SessionName.Get();
 
 			FUnregisterPlayersCallback* CallbackObj = new FUnregisterPlayersCallback();
 			CallbackObj->CallbackLambda = [this, SessionName, UnregisteredPlayers = TArray<FUniqueNetIdRef>(Players)](const EOS_Sessions_UnregisterPlayersCallbackInfo* Data)
