[EOS OSS] Register all remote players who are already in a session as the local player joins it.
Previously players would only get registered that join a session after the local player was already in it (through FOnlineSessionEOS::OnMemberStatusReceived).

diff --git a/Engine/Plugins/Online/OnlineSubsystemEOS/Source/OnlineSubsystemEOS/Private/OnlineSessionEOS.cpp b/Engine/Plugins/Online/OnlineSubsystemEOS/Source/OnlineSubsystemEOS/Private/OnlineSessionEOS.cpp
--- a/Engine/Plugins/Online/OnlineSubsystemEOS/Source/OnlineSubsystemEOS/Private/OnlineSessionEOS.cpp
+++ b/Engine/Plugins/Online/OnlineSubsystemEOS/Source/OnlineSubsystemEOS/Private/OnlineSessionEOS.cpp
@@ -700,27 +700,7 @@ void FOnlineSessionEOS::OnMemberStatusReceived(const EOS_LobbyId& LobbyId, const
 		{
 		case EOS_ELobbyMemberStatus::EOS_LMS_JOINED:
 			{
-				EOSSubsystem->UserManager->GetEpicAccountIdAsync(TargetUserId, [this, LobbyNetId](const EOS_ProductUserId& ProductUserId, EOS_EpicAccountId& EpicAccountId)
-					{
-						FNamedOnlineSession* Session = GetNamedSessionFromLobbyId(*LobbyNetId);
-						if (Session)
-						{
-							FUniqueNetIdPtr UniqueNetId = EOSSubsystem->UserManager->CreateUniquePlayerId(MakeNetIdStringFromIds(EpicAccountId, ProductUserId));
-
-							RegisterPlayer(Session->SessionName, *UniqueNetId, false);
-
-							// Right after registering the user, we query their member settings
-							const FTCHARToUTF8 Utf8LobbyId(*LobbyNetId->ToString());
-							OnLobbyMemberUpdateReceived((EOS_LobbyId)Utf8LobbyId.Get(), ProductUserId);
-
-							// Maybe redundant with the register player signal?
-							TriggerOnSessionParticipantsChangeDelegates(Session->SessionName, *UniqueNetId, true);
-						}
-						else
-						{
-							UE_LOG_ONLINE(Warning, TEXT("[FOnlineSessionEOS::OnMemberStatusReceived] Unable to retrieve session with LobbyId %s"), *LobbyNetId->ToString());
-						}
-					});
+				AddLobbyMember(LobbyNetId, TargetUserId);
 			}
 
 			break;
@@ -3425,7 +3405,7 @@ uint32 FOnlineSessionEOS::JoinLobbySession(int32 PlayerNum, FNamedOnlineSession*
 
 			FLobbyJoinedCallback* CallbackObj = new FLobbyJoinedCallback();
 			LobbyJoinedCallback = CallbackObj;
-			CallbackObj->CallbackLambda = [this, SessionName, LocalUserNetId](const EOS_Lobby_JoinLobbyCallbackInfo* Data)
+			CallbackObj->CallbackLambda = [this, SessionName, LocalUserNetId, LobbyDetailsHandle](const EOS_Lobby_JoinLobbyCallbackInfo* Data)
 			{
 				FNamedOnlineSession* Session = GetNamedSession(SessionName);
 				if (Session)
@@ -3436,6 +3416,21 @@ uint32 FOnlineSessionEOS::JoinLobbySession(int32 PlayerNum, FNamedOnlineSession*
 						UE_LOG_ONLINE_SESSION(Verbose, TEXT("[FOnlineSessionEOS::JoinLobbySession] JoinLobby was successful. LobbyId is %d."), Data->LobbyId);
 
 						BeginSessionAnalytics(Session);
+						
+						// Initialize the OSS's member list with the current member list of the lobby.
+						const FUniqueNetIdEOSRef LobbyNetId = FUniqueNetIdEOS::Create(UTF8_TO_TCHAR(Data->LobbyId));
+						EOS_LobbyDetails_GetMemberCountOptions GetMemberCountOptions = { 0 };
+						GetMemberCountOptions.ApiVersion = EOS_LOBBYDETAILS_GETMEMBERCOUNT_API_LATEST;
+						const uint32 MemberCount = EOS_LobbyDetails_GetMemberCount(LobbyDetailsHandle, &GetMemberCountOptions);
+						
+						EOS_LobbyDetails_GetMemberByIndexOptions GetMemberOptions = { 0 };
+						GetMemberOptions.ApiVersion = EOS_LOBBYDETAILS_GETMEMBERBYINDEX_API_LATEST;
+						for(uint32 MemberIndex = 0; MemberIndex < MemberCount; ++MemberIndex)
+						{
+							GetMemberOptions.MemberIndex = MemberIndex;
+							const EOS_ProductUserId ProductUserId = EOS_LobbyDetails_GetMemberByIndex(LobbyDetailsHandle, &GetMemberOptions);
+							AddLobbyMember(LobbyNetId, ProductUserId);
+						}
 
 #if WITH_EOS_RTC
 						if (FEOSVoiceChatUser* VoiceChatUser = static_cast<FEOSVoiceChatUser*>(EOSSubsystem->GetEOSVoiceChatUserInterface(*LocalUserNetId)))
@@ -3622,6 +3617,31 @@ void FOnlineSessionEOS::SetLobbyAttributes(EOS_HLobbyModification LobbyModificat
 	}
 }
 
+void FOnlineSessionEOS::AddLobbyMember(const FUniqueNetIdEOSRef LobbyNetId, const EOS_ProductUserId& TargetUserId)
+{
+	EOSSubsystem->UserManager->GetEpicAccountIdAsync(TargetUserId, [this, LobbyNetId](const EOS_ProductUserId& ProductUserId, EOS_EpicAccountId& EpicAccountId)
+		{
+			if (const FNamedOnlineSession* Session = GetNamedSessionFromLobbyId(*LobbyNetId))
+			{
+				const FUniqueNetIdPtr UniqueNetId = EOSSubsystem->UserManager->CreateUniquePlayerId(MakeNetIdStringFromIds(EpicAccountId, ProductUserId));
+
+				RegisterPlayer(Session->SessionName, *UniqueNetId, false);
+
+				// Right after registering the user, we query their member settings
+				const FTCHARToUTF8 Utf8LobbyId(*LobbyNetId->ToString());
+				OnLobbyMemberUpdateReceived(static_cast<EOS_LobbyId>(Utf8LobbyId.Get()), ProductUserId);
+
+				// Maybe redundant with the register player signal?
+				TriggerOnSessionParticipantsChangeDelegates(Session->SessionName, *UniqueNetId, true);
+			}
+			else
+			{
+				UE_LOG_ONLINE(Warning, TEXT("[FOnlineSessionEOS::AddLobbyMember] Unable to retrieve session with LobbyId %s"), *LobbyNetId->ToString());
+			}
+		});
+}
+
+
 uint32 FOnlineSessionEOS::UpdateLobbySession(FNamedOnlineSession* Session)
 {
 	check(Session != nullptr);
diff --git a/Engine/Plugins/Online/OnlineSubsystemEOS/Source/OnlineSubsystemEOS/Private/OnlineSessionEOS.h b/Engine/Plugins/Online/OnlineSubsystemEOS/Source/OnlineSubsystemEOS/Private/OnlineSessionEOS.h
--- a/Engine/Plugins/Online/OnlineSubsystemEOS/Source/OnlineSubsystemEOS/Private/OnlineSessionEOS.h
+++ b/Engine/Plugins/Online/OnlineSubsystemEOS/Source/OnlineSubsystemEOS/Private/OnlineSessionEOS.h
@@ -230,6 +230,7 @@ private:
 	void SetLobbyAttributes(EOS_HLobbyModification LobbyModificationHandle, FNamedOnlineSession* Session);
 	void AddLobbyAttribute(EOS_HLobbyModification LobbyModificationHandle, const EOS_Lobby_AttributeData* Attribute);
 	void AddLobbyMemberAttribute(EOS_HLobbyModification LobbyModificationHandle, const EOS_Lobby_AttributeData* Attribute);
+	void AddLobbyMember(const FUniqueNetIdEOSRef LobbyNetId, const EOS_ProductUserId& TargetUserId);
 
 	// Lobby search
 	void AddLobbySearchAttribute(EOS_HLobbySearch LobbySearchHandle, const EOS_Lobby_AttributeData* Attribute, EOS_EOnlineComparisonOp ComparisonOp);
