Moved EOSPlus login settings to their own category, expanded descriptions.

diff --git a/Engine/Plugins/Online/OnlineSubsystemEOS/Source/OnlineSubsystemEOS/Public/EOSSettings.h b/Engine/Plugins/Online/OnlineSubsystemEOS/Source/OnlineSubsystemEOS/Public/EOSSettings.h
--- a/Engine/Plugins/Online/OnlineSubsystemEOS/Source/OnlineSubsystemEOS/Public/EOSSettings.h
+++ b/Engine/Plugins/Online/OnlineSubsystemEOS/Source/OnlineSubsystemEOS/Public/EOSSettings.h
@@ -132,11 +132,11 @@ public:
 	TArray<FArtifactSettings> Artifacts;
 
 	/** Set to true to have Epic Accounts used (friends list will be unified with the default platform) */
-	UPROPERTY(Config, EditAnywhere, BlueprintReadOnly, Category="Crossplay Settings", DisplayName="Use Epic Account Services")
+	UPROPERTY(Config, EditAnywhere, BlueprintReadOnly, Category="EOSPlus Login Settings", DisplayName="Use Epic Account for EOS login (requires account linking)")
 	bool bUseEAS = false;
 
 	/** Set to true to have EOS Connect APIs used to link accounts for crossplay */
-	UPROPERTY(Config, EditAnywhere, BlueprintReadOnly, Category="Crossplay Settings", DisplayName="Use Crossplatform User IDs")
+	UPROPERTY(Config, EditAnywhere, BlueprintReadOnly, Category="EOSPlus Login Settings", DisplayName="Use Crossplatform User IDs for EOS Login (doesn't use Epic Account)")
 	bool bUseEOSConnect = false;
 
 	/** Set to true to write stats to EOS as well as the default platform */
