Modified FUserManagerEOS to properly return true/false in some UI request functions.

diff --git a/Engine/Plugins/Online/OnlineSubsystemEOS/Source/OnlineSubsystemEOS/Private/UserManagerEOS.cpp b/Engine/Plugins/Online/OnlineSubsystemEOS/Source/OnlineSubsystemEOS/Private/UserManagerEOS.cpp
--- a/Engine/Plugins/Online/OnlineSubsystemEOS/Source/OnlineSubsystemEOS/Private/UserManagerEOS.cpp
+++ b/Engine/Plugins/Online/OnlineSubsystemEOS/Source/OnlineSubsystemEOS/Private/UserManagerEOS.cpp
@@ -1547,12 +1547,15 @@ bool FUserManagerEOS::ShowWebURL(const FString& Url, const FShowWebUrlParams& Sh
 {
 	UE_LOG_ONLINE_EXTERNALUI(Warning, TEXT("[FUserManagerEOS::ShowWebURL] This method is not implemented."));
 
-	EOSSubsystem->ExecuteNextTick([this, Delegate]()
-		{
-			Delegate.ExecuteIfBound(FString());
-		});
-
-	return true;
+	if (Delegate.IsBound())
+	{
+		EOSSubsystem->ExecuteNextTick([this, Delegate]()
+			{
+				Delegate.ExecuteIfBound(FString());
+			});
+		return true;
+	}
+	return false;
 }
 
 bool FUserManagerEOS::CloseWebURL()
@@ -1566,12 +1569,15 @@ bool FUserManagerEOS::ShowProfileUI(const FUniqueNetId& Requestor, const FUnique
 {
 	UE_LOG_ONLINE_EXTERNALUI(Warning, TEXT("[FUserManagerEOS::ShowProfileUI] This method is not implemented."));
 
-	EOSSubsystem->ExecuteNextTick([this, Delegate]()
-		{
-			Delegate.ExecuteIfBound();
-		});
-
-	return true;
+	if (Delegate.IsBound())
+	{
+		EOSSubsystem->ExecuteNextTick([this, Delegate]()
+			{
+				Delegate.ExecuteIfBound();
+			});
+		return true;
+	}
+	return false;
 }
 
 bool FUserManagerEOS::ShowAccountUpgradeUI(const FUniqueNetId& UniqueId)
@@ -1585,24 +1591,30 @@ bool FUserManagerEOS::ShowStoreUI(int32 LocalUserNum, const FShowStoreParams& Sh
 {
 	UE_LOG_ONLINE_EXTERNALUI(Warning, TEXT("[FUserManagerEOS::ShowStoreUI] This method is not implemented."));
 
-	EOSSubsystem->ExecuteNextTick([this, Delegate]()
-		{
-			Delegate.ExecuteIfBound(false);
-		});
-
-	return true;
+	if (Delegate.IsBound())
+	{
+		EOSSubsystem->ExecuteNextTick([this, Delegate]()
+			{
+				Delegate.ExecuteIfBound(false);
+			});
+		return true;
+	}
+	return false;
 }
 
 bool FUserManagerEOS::ShowSendMessageUI(int32 LocalUserNum, const FShowSendMessageParams& ShowParams, const FOnShowSendMessageUIClosedDelegate& Delegate)
 {
 	UE_LOG_ONLINE_EXTERNALUI(Warning, TEXT("[FUserManagerEOS::ShowSendMessageUI] This method is not implemented."));
 
-	EOSSubsystem->ExecuteNextTick([this, Delegate]()
-		{
-			Delegate.ExecuteIfBound(false);
-		});
-
-	return true;
+	if (Delegate.IsBound())
+	{
+		EOSSubsystem->ExecuteNextTick([this, Delegate]()
+			{
+				Delegate.ExecuteIfBound(false);
+			});
+		return true;
+	}
+	return false;
 }
 
 // ~IOnlineExternalUI Interface
