Fix for EOSSDK requests hanging after applying a config update.

Merge + integration fix by [REDACTED_PERSON]

diff --git a/Engine/Plugins/Online/EOSShared/Source/EOSShared/Public/EOSSharedTypes.h b/Engine/Plugins/Online/EOSShared/Source/EOSShared/Public/EOSSharedTypes.h
--- a/Engine/Plugins/Online/EOSShared/Source/EOSShared/Public/EOSSharedTypes.h
+++ b/Engine/Plugins/Online/EOSShared/Source/EOSShared/Public/EOSSharedTypes.h
@@ -24,6 +24,7 @@ class EOSSHARED_API FCallbackBase
 public:
 	virtual ~FCallbackBase() {}
 	static bool ShouldCancelAllCallbacks() { return FCallbackBase::bShouldCancelAllCallbacks; }
+	static void EnableAllCallbacks() { FCallbackBase::bShouldCancelAllCallbacks = false; }
 	static void CancelAllCallbacks() { FCallbackBase::bShouldCancelAllCallbacks = true; }
 };
 
diff --git a/Engine/Plugins/Online/OnlineSubsystemEOS/Source/OnlineSubsystemEOS/Private/OnlineSubsystemEOS.cpp b/Engine/Plugins/Online/OnlineSubsystemEOS/Source/OnlineSubsystemEOS/Private/OnlineSubsystemEOS.cpp
--- a/Engine/Plugins/Online/OnlineSubsystemEOS/Source/OnlineSubsystemEOS/Private/OnlineSubsystemEOS.cpp
+++ b/Engine/Plugins/Online/OnlineSubsystemEOS/Source/OnlineSubsystemEOS/Private/OnlineSubsystemEOS.cpp
@@ -249,6 +249,8 @@ bool FOnlineSubsystemEOS::PlatformCreate()
 
 bool FOnlineSubsystemEOS::Init()
 {
+	FCallbackBase::EnableAllCallbacks();
+
 	// Determine if we are the default and if we're the platform OSS
 	FString DefaultOSS;
 	GConfig->GetString(TEXT("OnlineSubsystem"), TEXT("DefaultPlatformService"), DefaultOSS, GEngineIni);
