Update for CreateLobbyOptions fields: Addition of lobby id override and lobby host migration options.
Removed hardcoded search parameters from sessions and lobbies.
Refactored DestroyLobbySession to make use of lobby host migration logic.

diff --git a/Engine/Plugins/Online/OnlineSubsystem/Source/Public/OnlineSessionSettings.h b/Engine/Plugins/Online/OnlineSubsystem/Source/Public/OnlineSessionSettings.h
--- a/Engine/Plugins/Online/OnlineSubsystem/Source/Public/OnlineSessionSettings.h
+++ b/Engine/Plugins/Online/OnlineSubsystem/Source/Public/OnlineSessionSettings.h
@@ -256,6 +256,9 @@ public:
 	bool bUseLobbiesIfAvailable;
 	/** Whether to create (and auto join) a voice chat room for the lobby, if the platform supports it */
 	bool bUseLobbiesVoiceChatIfAvailable;
+	/** Manual override for the Session Id instead of having one assigned by the backend. Its size may be restricted depending on the platform */
+	FString SessionIdOverride;
+
 	/** Used to keep different builds from seeing each other during searches */
 	int32 BuildUniqueId;
 	/** Array of custom session settings */
diff --git a/Engine/Plugins/Online/OnlineSubsystemEOS/Source/OnlineSubsystemEOS/Private/OnlineSessionEOS.cpp b/Engine/Plugins/Online/OnlineSubsystemEOS/Source/OnlineSubsystemEOS/Private/OnlineSessionEOS.cpp
--- a/Engine/Plugins/Online/OnlineSubsystemEOS/Source/OnlineSubsystemEOS/Private/OnlineSessionEOS.cpp
+++ b/Engine/Plugins/Online/OnlineSubsystemEOS/Source/OnlineSubsystemEOS/Private/OnlineSessionEOS.cpp
@@ -1039,11 +1039,6 @@ void FOnlineSessionEOS::AddAttribute(EOS_HSessionModification SessionModHandle,
 
 void FOnlineSessionEOS::SetAttributes(EOS_HSessionModification SessionModHandle, FNamedOnlineSession* Session)
 {
-	// The first will let us find it on session searches
-	const FString SearchPresence(SEARCH_PRESENCE.ToString());
-	const FAttributeOptions SearchPresenceAttribute(TCHAR_TO_UTF8(*SearchPresence), true);
-	AddAttribute(SessionModHandle, &SearchPresenceAttribute);
-
 	FAttributeOptions Opt1("NumPrivateConnections", Session->SessionSettings.NumPrivateConnections);
 	AddAttribute(SessionModHandle, &Opt1);
 
@@ -3310,6 +3305,8 @@ uint32 FOnlineSessionEOS::CreateLobbySession(int32 HostingPlayerNum, FNamedOnlin
 
 	const EOS_ProductUserId LocalProductUserId = EOSSubsystem->UserManager->GetLocalProductUserId(HostingPlayerNum);
 	const FUniqueNetIdPtr LocalUserNetId = EOSSubsystem->UserManager->GetLocalUniqueNetIdEOS(HostingPlayerNum);
+	bool bUseHostMigration = true;
+	Session->SessionSettings.Get(SETTING_HOST_MIGRATION, bUseHostMigration);
 
 	EOS_Lobby_CreateLobbyOptions CreateLobbyOptions = { 0 };
 	CreateLobbyOptions.ApiVersion = EOS_LOBBY_CREATELOBBY_API_LATEST;
@@ -3319,9 +3316,19 @@ uint32 FOnlineSessionEOS::CreateLobbySession(int32 HostingPlayerNum, FNamedOnlin
 	CreateLobbyOptions.bPresenceEnabled = Session->SessionSettings.bUsesPresence;
 	CreateLobbyOptions.bAllowInvites = Session->SessionSettings.bAllowInvites;
 	CreateLobbyOptions.BucketId = BucketIdAnsi;
+	CreateLobbyOptions.bDisableHostMigration = !bUseHostMigration;
 #if WITH_EOS_RTC
 	CreateLobbyOptions.bEnableRTCRoom = Session->SessionSettings.bUseLobbiesVoiceChatIfAvailable;
 #endif
+	const FTCHARToUTF8 Utf8SessionIdOverride(*Session->SessionSettings.SessionIdOverride);
+	if (Session->SessionSettings.SessionIdOverride.Len() >= EOS_LOBBY_MIN_LOBBYIDOVERRIDE_LENGTH && Session->SessionSettings.SessionIdOverride.Len() <= EOS_LOBBY_MAX_LOBBYIDOVERRIDE_LENGTH)
+	{
+		CreateLobbyOptions.LobbyId = Utf8SessionIdOverride.Get();
+	}
+	else if (!Session->SessionSettings.SessionIdOverride.IsEmpty())
+	{
+		UE_LOG_ONLINE_SESSION(Warning, TEXT("[FOnlineSessionEOS::CreateLobbySession] Session setting SessionIdOverride is of invalid length [%d]. Valid length range is between %d and %d."), Session->SessionSettings.SessionIdOverride.Len(), EOS_LOBBY_MIN_LOBBYIDOVERRIDE_LENGTH, EOS_LOBBY_MAX_LOBBYIDOVERRIDE_LENGTH)
+	}
 
 	/*When the operation finishes, the EOS_Lobby_OnCreateLobbyCallback will run with an EOS_Lobby_CreateLobbyCallbackInfo data structure.
 	If the data structure's ResultCode field indicates success, its LobbyId field contains the new lobby's ID value, which we will need to interact with the lobby further.*/
@@ -3545,16 +3552,6 @@ void FOnlineSessionEOS::SetLobbyAttributes(EOS_HLobbyModification LobbyModificat
 {
 	check(Session != nullptr);
 
-	// The first will let us find it on session searches
-	const FString SearchPresence(SEARCH_PRESENCE.ToString());
-	const FLobbyAttributeOptions SearchPresenceAttribute(TCHAR_TO_UTF8(*SearchPresence), true);
-	AddLobbyAttribute(LobbyModificationHandle, &SearchPresenceAttribute);
-
-	// The second will let us find it on lobby searches
-	const FString SearchLobbies(SEARCH_LOBBIES.ToString());
-	const FLobbyAttributeOptions SearchLobbiesAttribute(TCHAR_TO_UTF8(*SearchLobbies), true);
-	AddLobbyAttribute(LobbyModificationHandle, &SearchLobbiesAttribute);
-
 	// We set the session's owner id and name
 	const FLobbyAttributeOptions OwnerId("OwningUserId", TCHAR_TO_UTF8(*Session->OwningUserId->ToString()));
 	AddLobbyAttribute(LobbyModificationHandle, &OwnerId);
@@ -3744,109 +3741,55 @@ uint32 FOnlineSessionEOS::DestroyLobbySession(FNamedOnlineSession* Session, cons
 		FOnlineSessionInfoEOS* SessionInfo = (FOnlineSessionInfoEOS*)(Session->SessionInfo.Get());
 		check(Session->SessionSettings.bUseLobbiesIfAvailable); // We check if it's a lobby session
 
-		FName SessionName = Session->SessionName;
-		const FUniqueNetIdPtr OwningNetId = Session->OwningUserId;
-		const FUniqueNetIdPtr LocalNetId = EOSSubsystem->UserManager->GetUniquePlayerId(EOSSubsystem->UserManager->GetDefaultLocalUser());
+		// EOS will use the host migration setting to decide if the lobby is destroyed if it's the owner leaving
+		EOS_Lobby_LeaveLobbyOptions LeaveOptions = { 0 };
+		LeaveOptions.ApiVersion = EOS_LOBBY_LEAVELOBBY_API_LATEST;
+		const FTCHARToUTF8 Utf8LobbyId(*SessionInfo->GetSessionId().ToString());
+		LeaveOptions.LobbyId = (EOS_LobbyId)Utf8LobbyId.Get();
+		LeaveOptions.LocalUserId = EOSSubsystem->UserManager->GetLocalProductUserId(EOSSubsystem->UserManager->GetDefaultLocalUser()); // Maybe not split screen friendly
 
-		// If we are the owner of the lobby we will destroy it, since we will trigger an OSS session destruction anyway, forcing everyone else to leave the lobby too
-		if (OwningNetId != nullptr && LocalNetId  != nullptr && *LocalNetId == *OwningNetId)
+		FName SessionName = Session->SessionName;
+		FLobbyLeftCallback* LeaveCallbackObj = new FLobbyLeftCallback();
+		LobbyLeftCallback = LeaveCallbackObj;
+		LeaveCallbackObj->CallbackLambda = [this, SessionName, CompletionDelegate](const EOS_Lobby_LeaveLobbyCallbackInfo* Data)
 		{
-			// Destroy Lobby
-			EOS_Lobby_DestroyLobbyOptions DestroyOptions = { 0 };
-			DestroyOptions.ApiVersion = EOS_LOBBY_DESTROYLOBBY_API_LATEST;
-			const FTCHARToUTF8 Utf8LobbyId(*SessionInfo->GetSessionId().ToString());
-			DestroyOptions.LobbyId = (EOS_LobbyId)Utf8LobbyId.Get();
-			DestroyOptions.LocalUserId = EOSSubsystem->UserManager->GetLocalProductUserId(EOSSubsystem->UserManager->GetDefaultLocalUser()); // Maybe not split screen friendly
-
-			FLobbyDestroyedCallback* DestroyCallbackObj = new FLobbyDestroyedCallback();
-			LobbyDestroyedCallback = DestroyCallbackObj;
-			DestroyCallbackObj->CallbackLambda = [this, SessionName, CompletionDelegate](const EOS_Lobby_DestroyLobbyCallbackInfo* Data)
+			FNamedOnlineSession* LobbySession = GetNamedSession(SessionName);
+			if (LobbySession)
 			{
-				FNamedOnlineSession* LobbySession = GetNamedSession(SessionName);
-				if (LobbySession)
+				bool bWasSuccessful = Data->ResultCode == EOS_EResult::EOS_Success;
+				if (bWasSuccessful)
 				{
-					bool bWasSuccessful = Data->ResultCode == EOS_EResult::EOS_Success;
-					if (!bWasSuccessful)
-					{
-						UE_LOG_ONLINE_SESSION(Warning, TEXT("[FOnlineSessionEOS::DestroyLobbySession] DestroyLobby not successful. Finished with EOS_EResult %s"), ANSI_TO_TCHAR(EOS_EResult_ToString(Data->ResultCode)));
-					}
-
-#if WITH_EOS_RTC
-					if (FEOSVoiceChatUser* VoiceChatUser = static_cast<FEOSVoiceChatUser*>(EOSSubsystem->GetEOSVoiceChatUserInterface(*EOSSubsystem->UserManager->GetLocalUniqueNetIdEOS())))
-					{
-						VoiceChatUser->RemoveLobbyRoom(UTF8_TO_TCHAR(Data->LobbyId));
-					}
-#endif
-
-					EndSessionAnalytics();
-
-					LobbySession->SessionState = EOnlineSessionState::NoSession;
-
-					RemoveNamedSession(SessionName);
-
-					CompletionDelegate.ExecuteIfBound(SessionName, bWasSuccessful);
-					TriggerOnDestroySessionCompleteDelegates(SessionName, bWasSuccessful);
+					UE_LOG_ONLINE_SESSION(Verbose, TEXT("[FOnlineSessionEOS::DestroyLobbySession] LeaveLobby was successful. LobbyId is %s."), Data->LobbyId);
 				}
 				else
 				{
-					UE_LOG_ONLINE_SESSION(Warning, TEXT("[FOnlineSessionEOS::DestroyLobbySession] Unable to find session %s"), *SessionName.ToString());
-					TriggerOnDestroySessionCompleteDelegates(SessionName, false);
+					UE_LOG_ONLINE_SESSION(Warning, TEXT("[FOnlineSessionEOS::DestroyLobbySession] LeaveLobby not successful. Finished with EOS_EResult %s"), ANSI_TO_TCHAR(EOS_EResult_ToString(Data->ResultCode)));
 				}
-			};
-
-			EOS_Lobby_DestroyLobby(LobbyHandle, &DestroyOptions, DestroyCallbackObj, DestroyCallbackObj->GetCallbackPtr());
-		}
-		else
-		{
-			// Leave Lobby
-			EOS_Lobby_LeaveLobbyOptions LeaveOptions = { 0 };
-			LeaveOptions.ApiVersion = EOS_LOBBY_LEAVELOBBY_API_LATEST;
-			const FTCHARToUTF8 Utf8LobbyId(*SessionInfo->GetSessionId().ToString());
-			LeaveOptions.LobbyId = (EOS_LobbyId)Utf8LobbyId.Get();
-			LeaveOptions.LocalUserId = EOSSubsystem->UserManager->GetLocalProductUserId(EOSSubsystem->UserManager->GetDefaultLocalUser()); // Maybe not split screen friendly
-
-			FLobbyLeftCallback* LeaveCallbackObj = new FLobbyLeftCallback();
-			LobbyLeftCallback = LeaveCallbackObj;
-			LeaveCallbackObj->CallbackLambda = [this, SessionName, CompletionDelegate](const EOS_Lobby_LeaveLobbyCallbackInfo* Data)
-			{
-				FNamedOnlineSession* LobbySession = GetNamedSession(SessionName);
-				if (LobbySession)
-				{
-					bool bWasSuccessful = Data->ResultCode == EOS_EResult::EOS_Success;
-					if (bWasSuccessful)
-					{
-						UE_LOG_ONLINE_SESSION(Verbose, TEXT("[FOnlineSessionEOS::DestroyLobbySession] LeaveLobby was successful. LobbyId is %s."), Data->LobbyId);
-					}
-					else
-					{
-						UE_LOG_ONLINE_SESSION(Warning, TEXT("[FOnlineSessionEOS::DestroyLobbySession] LeaveLobby not successful. Finished with EOS_EResult %s"), ANSI_TO_TCHAR(EOS_EResult_ToString(Data->ResultCode)));
-					}
 
 #if WITH_EOS_RTC
-					if (FEOSVoiceChatUser* VoiceChatUser = static_cast<FEOSVoiceChatUser*>(EOSSubsystem->GetEOSVoiceChatUserInterface(*EOSSubsystem->UserManager->GetLocalUniqueNetIdEOS())))
-					{
-						VoiceChatUser->RemoveLobbyRoom(UTF8_TO_TCHAR(Data->LobbyId));
-					}
+				if (FEOSVoiceChatUser* VoiceChatUser = static_cast<FEOSVoiceChatUser*>(EOSSubsystem->GetEOSVoiceChatUserInterface(*EOSSubsystem->UserManager->GetLocalUniqueNetIdEOS())))
+				{
+					VoiceChatUser->RemoveLobbyRoom(UTF8_TO_TCHAR(Data->LobbyId));
+				}
 #endif
 
-					EndSessionAnalytics();
+				EndSessionAnalytics();
 
-					LobbySession->SessionState = EOnlineSessionState::NoSession;
+				LobbySession->SessionState = EOnlineSessionState::NoSession;
 
-					RemoveNamedSession(SessionName);
+				RemoveNamedSession(SessionName);
 
-					CompletionDelegate.ExecuteIfBound(SessionName, bWasSuccessful);
-					TriggerOnDestroySessionCompleteDelegates(SessionName, bWasSuccessful);
-				}
-				else
-				{
-					UE_LOG_ONLINE_SESSION(Warning, TEXT("[FOnlineSessionEOS::DestroyLobbySession] Unable to find session %s"), *SessionName.ToString());
-					TriggerOnDestroySessionCompleteDelegates(SessionName, false);
-				}
-			};
+				CompletionDelegate.ExecuteIfBound(SessionName, bWasSuccessful);
+				TriggerOnDestroySessionCompleteDelegates(SessionName, bWasSuccessful);
+			}
+			else
+			{
+				UE_LOG_ONLINE_SESSION(Warning, TEXT("[FOnlineSessionEOS::DestroyLobbySession] Unable to find session %s"), *SessionName.ToString());
+				TriggerOnDestroySessionCompleteDelegates(SessionName, false);
+			}
+		};
 
-			EOS_Lobby_LeaveLobby(LobbyHandle, &LeaveOptions, LeaveCallbackObj, LeaveCallbackObj->GetCallbackPtr());
-		}
+		EOS_Lobby_LeaveLobby(LobbyHandle, &LeaveOptions, LeaveCallbackObj, LeaveCallbackObj->GetCallbackPtr());
 
 		Result = ONLINE_IO_PENDING;
 	}
@@ -4040,6 +3983,7 @@ void FOnlineSessionEOS::CopyLobbyData(EOS_HLobbyDetails LobbyDetailsHandle, EOS_
 {
 	OutSession.SessionSettings.bUseLobbiesIfAvailable = true;
 	OutSession.SessionSettings.bIsLANMatch = false;
+	OutSession.SessionSettings.Set(SETTING_HOST_MIGRATION, LobbyDetailsInfo->bAllowHostMigration, EOnlineDataAdvertisementType::DontAdvertise);
 #if WITH_EOS_RTC
 	OutSession.SessionSettings.bUseLobbiesVoiceChatIfAvailable = LobbyDetailsInfo->bRTCRoomEnabled == EOS_TRUE;
 #endif
