Fix for FAuthCredentials using incorrect buffer size when calling EOS_ByteArray_ToString causing it to return EOS_LimitExceeded

- Inside FAuthCredentials::Init, we call EOS_ByteArray_ToString to convert the token bytes to a string passing in the TokenAnsi buffer, but we pass its size as EOS_OSS_STRING_BUFFER_LENGTH (257) instead of its actual size EOS_MAX_TOKEN_SIZE (4096). When using a Steam token, this causes EOS_ByteArray_ToString to return EOS_LimitExceeded because it needs a min size of 287

[REDACTED_HANDLE] [REDACTED_HANDLE]

diff --git a/Engine/Plugins/Online/OnlineSubsystemEOS/Source/OnlineSubsystemEOS/Private/UserManagerEOS.cpp b/Engine/Plugins/Online/OnlineSubsystemEOS/Source/OnlineSubsystemEOS/Private/UserManagerEOS.cpp
--- a/Engine/Plugins/Online/OnlineSubsystemEOS/Source/OnlineSubsystemEOS/Private/UserManagerEOS.cpp
+++ b/Engine/Plugins/Online/OnlineSubsystemEOS/Source/OnlineSubsystemEOS/Private/UserManagerEOS.cpp
@@ -339,7 +339,7 @@ struct FAuthCredentials :
 		Id = IdAnsi;
 		Token = TokenAnsi;
 
-		uint32_t InOutBufferLength = EOS_OSS_STRING_BUFFER_LENGTH;
+		uint32_t InOutBufferLength = EOS_MAX_TOKEN_SIZE;
 		EOS_ByteArray_ToString(InToken.GetData(), InToken.Num(), TokenAnsi, &InOutBufferLength);
 	}
 
