PR #8840
Added support for GetLinkedAccountAuthToken in OnlineSubsystemEOS
(Contributed by [REDACTED_PERSON])

diff --git a/Engine/Plugins/Online/OnlineSubsystemEOS/Source/OnlineSubsystemEOS/Private/UserManagerEOS.cpp b/Engine/Plugins/Online/OnlineSubsystemEOS/Source/OnlineSubsystemEOS/Private/UserManagerEOS.cpp
--- a/Engine/Plugins/Online/OnlineSubsystemEOS/Source/OnlineSubsystemEOS/Private/UserManagerEOS.cpp
+++ b/Engine/Plugins/Online/OnlineSubsystemEOS/Source/OnlineSubsystemEOS/Private/UserManagerEOS.cpp
@@ -1463,6 +1463,13 @@ FPlatformUserId FUserManagerEOS::GetPlatformUserIdFromUniqueNetId(const FUniqueN
 	return GetPlatformUserIdFromLocalUserNum(GetLocalUserNumFromUniqueNetId(UniqueNetId));
 }
 
+void FUserManagerEOS::GetLinkedAccountAuthToken(int32 LocalUserNum, const FOnGetLinkedAccountAuthTokenCompleteDelegate& Delegate) const
+{
+	FExternalAuthToken ExternalToken;
+	ExternalToken.TokenString = GetAuthToken(LocalUserNum);
+	Delegate.ExecuteIfBound(LocalUserNum, ExternalToken.HasTokenString(), ExternalToken);
+}
+
 void FUserManagerEOS::GetUserPrivilege(const FUniqueNetId& UserId, EUserPrivileges::Type Privilege, const FOnGetUserPrivilegeCompleteDelegate& Delegate)
 {
 	Delegate.ExecuteIfBound(UserId, Privilege, (uint32)EPrivilegeResults::NoFailures);
diff --git a/Engine/Plugins/Online/OnlineSubsystemEOS/Source/OnlineSubsystemEOS/Private/UserManagerEOS.h b/Engine/Plugins/Online/OnlineSubsystemEOS/Source/OnlineSubsystemEOS/Private/UserManagerEOS.h
--- a/Engine/Plugins/Online/OnlineSubsystemEOS/Source/OnlineSubsystemEOS/Private/UserManagerEOS.h
+++ b/Engine/Plugins/Online/OnlineSubsystemEOS/Source/OnlineSubsystemEOS/Private/UserManagerEOS.h
@@ -273,6 +273,7 @@ public:
 	virtual FString GetAuthType() const override;
 	virtual void RevokeAuthToken(const FUniqueNetId& LocalUserId, const FOnRevokeAuthTokenCompleteDelegate& Delegate) override;
 	virtual FPlatformUserId GetPlatformUserIdFromUniqueNetId(const FUniqueNetId& UniqueNetId) const override;
+	virtual void GetLinkedAccountAuthToken(int32 LocalUserNum, const FOnGetLinkedAccountAuthTokenCompleteDelegate& Delegate) const override;
 // ~IOnlineIdentity Interface
 	ELoginStatus::Type GetLoginStatus(const FUniqueNetIdEOS& UserId) const;
 
