[EOS OSS] Support serializing json-typed FVarientData.
The format is identical to string-typed FVarientData which is already supported.

diff --git a/Engine/Plugins/Online/OnlineSubsystemEOS/Source/OnlineSubsystemEOS/Private/OnlineSessionEOS.cpp b/Engine/Plugins/Online/OnlineSubsystemEOS/Source/OnlineSubsystemEOS/Private/OnlineSessionEOS.cpp
--- a/Engine/Plugins/Online/OnlineSubsystemEOS/Source/OnlineSubsystemEOS/Private/OnlineSessionEOS.cpp
+++ b/Engine/Plugins/Online/OnlineSubsystemEOS/Source/OnlineSubsystemEOS/Private/OnlineSessionEOS.cpp
@@ -57,6 +57,7 @@ bool IsSessionSettingTypeSupported(EOnlineKeyValuePairDataType::Type InType)
 		case EOnlineKeyValuePairDataType::String:
 		case EOnlineKeyValuePairDataType::Float:
 		case EOnlineKeyValuePairDataType::Bool:
+		case EOnlineKeyValuePairDataType::Json:
 		{
 			return true;
 		}
@@ -329,6 +330,7 @@ struct FLobbyAttributeOptions :
 			break;
 		}
 		case EOnlineKeyValuePairDataType::String:
+		case EOnlineKeyValuePairDataType::Json:
 		{
 			ValueType = EOS_ELobbyAttributeType::EOS_SAT_String;
 			Value.AsUtf8 = ValueAnsi;
