Initialised delegate to null and added checks to prevent crash.

diff --git a/Engine/Plugins/Online/OnlineSubsystemEOS/Source/OnlineSubsystemEOS/Private/UserManagerEOS.cpp b/Engine/Plugins/Online/OnlineSubsystemEOS/Source/OnlineSubsystemEOS/Private/UserManagerEOS.cpp
--- a/Engine/Plugins/Online/OnlineSubsystemEOS/Source/OnlineSubsystemEOS/Private/UserManagerEOS.cpp
+++ b/Engine/Plugins/Online/OnlineSubsystemEOS/Source/OnlineSubsystemEOS/Private/UserManagerEOS.cpp
@@ -144,27 +144,31 @@ FUserManagerEOS::FUserManagerEOS(FOnlineSubsystemEOS* InSubsystem)
 	// This delegate would cause a crash when running a dedicated server
 	if (!IsRunningDedicatedServer())
 	{
-	// Adding subscription to external ui display change event
-	EOS_UI_AddNotifyDisplaySettingsUpdatedOptions Options = {};
-	Options.ApiVersion = EOS_UI_ADDNOTIFYDISPLAYSETTINGSUPDATED_API_LATEST;
+		// Adding subscription to external ui display change event
+		EOS_UI_AddNotifyDisplaySettingsUpdatedOptions Options = {};
+		Options.ApiVersion = EOS_UI_ADDNOTIFYDISPLAYSETTINGSUPDATED_API_LATEST;
 
-	FOnDisplaySettingsUpdatedCallback* CallbackObj = new FOnDisplaySettingsUpdatedCallback();
-	DisplaySettingsUpdatedCallback = CallbackObj;
-	CallbackObj->CallbackLambda = [this](const EOS_UI_OnDisplaySettingsUpdatedCallbackInfo* Data)
-	{
-		TriggerOnExternalUIChangeDelegates((bool)Data->bIsVisible);
-	};
+		FOnDisplaySettingsUpdatedCallback* CallbackObj = new FOnDisplaySettingsUpdatedCallback();
+		DisplaySettingsUpdatedCallback = CallbackObj;
+		CallbackObj->CallbackLambda = [this](const EOS_UI_OnDisplaySettingsUpdatedCallbackInfo* Data)
+		{
+			TriggerOnExternalUIChangeDelegates((bool)Data->bIsVisible);
+		};
 
-	DisplaySettingsUpdatedId = EOS_UI_AddNotifyDisplaySettingsUpdated(EOSSubsystem->UIHandle, &Options, CallbackObj, CallbackObj->GetCallbackPtr());
-}
+		DisplaySettingsUpdatedId = EOS_UI_AddNotifyDisplaySettingsUpdated(EOSSubsystem->UIHandle, &Options, CallbackObj, CallbackObj->GetCallbackPtr());
+	}
 }
 
 FUserManagerEOS::~FUserManagerEOS()
 {
-	// Removing subscription to external ui display change event
-	EOS_UI_RemoveNotifyDisplaySettingsUpdated(EOSSubsystem->UIHandle, DisplaySettingsUpdatedId);
+	// This delegate would cause a crash when running a dedicated server
+	if (DisplaySettingsUpdatedId != EOS_INVALID_NOTIFICATIONID)
+	{
+		// Removing subscription to external ui display change event
+		EOS_UI_RemoveNotifyDisplaySettingsUpdated(EOSSubsystem->UIHandle, DisplaySettingsUpdatedId);
 
-	delete DisplaySettingsUpdatedCallback;
+		delete DisplaySettingsUpdatedCallback;
+	}
 }
 
 void FUserManagerEOS::LoginStatusChanged(const EOS_Auth_LoginStatusChangedCallbackInfo* Data)
diff --git a/Engine/Plugins/Online/OnlineSubsystemEOS/Source/OnlineSubsystemEOS/Private/UserManagerEOS.h b/Engine/Plugins/Online/OnlineSubsystemEOS/Source/OnlineSubsystemEOS/Private/UserManagerEOS.h
--- a/Engine/Plugins/Online/OnlineSubsystemEOS/Source/OnlineSubsystemEOS/Private/UserManagerEOS.h
+++ b/Engine/Plugins/Online/OnlineSubsystemEOS/Source/OnlineSubsystemEOS/Private/UserManagerEOS.h
@@ -474,8 +474,8 @@ private:
 	TMap<int32, TArray<ReadUserListInfo>> CachedReadUserListInfoForLocalUserMap;
 
 	/** Identifier for the external UI notification callback */
-	EOS_NotificationId DisplaySettingsUpdatedId;
-	FCallbackBase* DisplaySettingsUpdatedCallback;
+	EOS_NotificationId DisplaySettingsUpdatedId = EOS_INVALID_NOTIFICATIONID;
+	FCallbackBase* DisplaySettingsUpdatedCallback = nullptr;
 
 	/** Last Login Credentials used for a login attempt */
 	TMap<int32, TSharedRef<FOnlineAccountCredentials>> LocalUserNumToLastLoginCredentials;
