Added support for retrying platform session invites after login in EOSPlus.

diff --git a/Engine/Plugins/Online/OnlineSubsystemEOS/Source/OnlineSubsystemEOSPlus/Private/OnlineSessionEOSPlus.cpp b/Engine/Plugins/Online/OnlineSubsystemEOS/Source/OnlineSubsystemEOSPlus/Private/OnlineSessionEOSPlus.cpp
--- a/Engine/Plugins/Online/OnlineSubsystemEOS/Source/OnlineSubsystemEOSPlus/Private/OnlineSessionEOSPlus.cpp
+++ b/Engine/Plugins/Online/OnlineSubsystemEOS/Source/OnlineSubsystemEOSPlus/Private/OnlineSessionEOSPlus.cpp
@@ -125,6 +125,16 @@ TArray<FUniqueNetIdRef> FOnlineSessionEOSPlus::GetEOSNetIds(const TArray<FUnique
 	return EOSIds;
 }
 
+void FOnlineSessionEOSPlus::OnLoginComplete(int32 LocalUserNum, bool bWasSuccessful, const FUniqueNetId& UserId, const FString& Error)
+{
+	FUniqueNetIdRef BaseUserIdRef = GetNetIdPlus(UserId.ToString())->GetBaseNetId().ToSharedRef();
+
+	if (bWasSuccessful && PendingInviteResultsPerUser.Contains(BaseUserIdRef))
+	{
+		OnSessionUserInviteAcceptedBase(bWasSuccessful, LocalUserNum, BaseUserIdRef, *PendingInviteResultsPerUser[BaseUserIdRef]);
+	}
+}
+
 void FOnlineSessionEOSPlus::OnSessionUserInviteAcceptedBase(const bool bWasSuccessful, const int32 ControllerId, FUniqueNetIdPtr UserId, const FOnlineSessionSearchResult& InviteResult)
 {
 	if (!bWasSuccessful || !UserId.IsValid() || !InviteResult.IsValid())
@@ -151,15 +161,36 @@ void FOnlineSessionEOSPlus::OnSessionUserInviteAcceptedBase(const bool bWasSucce
 		if (!EOSUserId.IsValid())
 		{
 			UE_LOG_ONLINE(Error, TEXT("Failed to get EOS user id (%d)"), ControllerId);
+
+			FUniqueNetIdRef UserIdRef = UserId.ToSharedRef();
+
+			if (!PendingInviteResultsPerUser.Contains(UserIdRef))
+			{
+				UE_LOG_ONLINE(Verbose, TEXT("Retrying platform session invitation on successful login for user (%d)."), ControllerId);
+
+				PendingInviteResultsPerUser.Add(UserIdRef, MakeShared<FOnlineSessionSearchResult>(InviteResult));
+
+				EOSPlus->GetIdentityInterface()->AddOnLoginCompleteDelegate_Handle(ControllerId, FOnLoginCompleteDelegate::CreateRaw(this, &FOnlineSessionEOSPlus::OnLoginComplete));
+			}
+
 			return;
 		}
 		// Do a search for the EOS session
 		EOSSessionInterface->FindSessionById(*EOSUserId, *SessionId, *FUniqueNetIdString::EmptyId(),
 			FOnSingleSessionResultCompleteDelegate::CreateLambda([this, UserId](int32 LocalUserNum, bool bWasSuccessful, const FOnlineSessionSearchResult& EOSResult)
 			{
+				FUniqueNetIdRef UserIdRef = UserId.ToSharedRef();
+
 				FUniqueNetIdPtr PlusUserId = EOSPlus->UserInterfacePtr->GetUniquePlayerId(LocalUserNum);
 				// Now that we have the proper session trigger the invite
 				TriggerOnSessionUserInviteAcceptedDelegates(bWasSuccessful, LocalUserNum, PlusUserId, EOSResult);
+
+				if (PendingInviteResultsPerUser.Contains(UserIdRef))
+				{
+					EOSPlus->EosOSS->GetIdentityInterface()->ClearOnLoginCompleteDelegates(LocalUserNum, this);
+
+					PendingInviteResultsPerUser.Remove(UserIdRef);
+				}
 			}));
 	}
 	else
diff --git a/Engine/Plugins/Online/OnlineSubsystemEOS/Source/OnlineSubsystemEOSPlus/Private/OnlineSessionEOSPlus.h b/Engine/Plugins/Online/OnlineSubsystemEOS/Source/OnlineSubsystemEOSPlus/Private/OnlineSessionEOSPlus.h
--- a/Engine/Plugins/Online/OnlineSubsystemEOS/Source/OnlineSubsystemEOSPlus/Private/OnlineSessionEOSPlus.h
+++ b/Engine/Plugins/Online/OnlineSubsystemEOS/Source/OnlineSubsystemEOSPlus/Private/OnlineSessionEOSPlus.h
@@ -93,6 +93,10 @@ PACKAGE_SCOPE:
 	TArray<FUniqueNetIdRef> GetBaseNetIds(const TArray<FUniqueNetIdRef>& Players);
 	TArray<FUniqueNetIdRef> GetEOSNetIds(const TArray<FUniqueNetIdRef>& Players);
 
+	/** Support for platform session invitations, which we delay until login completion */
+	void OnLoginComplete(int32 LocalUserNum, bool bWasSuccessful, const FUniqueNetId& UserId, const FString& Error);
+	TUniqueNetIdMap<TSharedPtr<FOnlineSessionSearchResult>> PendingInviteResultsPerUser;
+
 	/** Reference to the owning EOS plus subsystem */
 	FOnlineSubsystemEOSPlus* EOSPlus;
 	bool bUseEOSSessions;
