Tied EOS callback lifetime to the object that requested the call

diff --git a/Engine/Plugins/Online/EOSShared/Source/EOSShared/Private/EOSShared.cpp b/Engine/Plugins/Online/EOSShared/Source/EOSShared/Private/EOSShared.cpp
--- a/Engine/Plugins/Online/EOSShared/Source/EOSShared/Private/EOSShared.cpp
+++ b/Engine/Plugins/Online/EOSShared/Source/EOSShared/Private/EOSShared.cpp
@@ -7,8 +7,6 @@
 
 DEFINE_LOG_CATEGORY(LogEOSSDK);
 
-bool FCallbackBase::bShouldCancelAllCallbacks = false;
-
 FString LexToString(const EOS_ProductUserId UserId)
 {
 	FString Result;
diff --git a/Engine/Plugins/Online/EOSShared/Source/EOSShared/Public/EOSSharedTypes.h b/Engine/Plugins/Online/EOSShared/Source/EOSShared/Public/EOSSharedTypes.h
--- a/Engine/Plugins/Online/EOSShared/Source/EOSShared/Public/EOSSharedTypes.h
+++ b/Engine/Plugins/Online/EOSShared/Source/EOSShared/Public/EOSSharedTypes.h
@@ -3,6 +3,7 @@
 #pragma once
 
 #include "CoreGlobals.h"
+#include "Templates/SharedPointer.h"
 
 #if WITH_EOS_SDK
 #if defined(EOS_PLATFORM_BASE_FILE_NAME)
@@ -19,26 +20,25 @@
 /** Used to store a pointer to the EOS callback object without knowing type */
 class EOSSHARED_API FCallbackBase
 {
-	static bool bShouldCancelAllCallbacks;
-
 public:
 	virtual ~FCallbackBase() {}
-	static bool ShouldCancelAllCallbacks() { return FCallbackBase::bShouldCancelAllCallbacks; }
-	static void EnableAllCallbacks() { FCallbackBase::bShouldCancelAllCallbacks = false; }
-	static void CancelAllCallbacks() { FCallbackBase::bShouldCancelAllCallbacks = true; }
 };
 
 #if WITH_EOS_SDK
 
 /** Class to handle all callbacks generically using a lambda to process callback results */
-template<typename CallbackFuncType, typename CallbackType>
+template<typename CallbackFuncType, typename CallbackType, typename OwningType>
 class TEOSGlobalCallback :
 	public FCallbackBase
 {
 public:
 	TFunction<void(const CallbackType*)> CallbackLambda;
 
-	TEOSGlobalCallback() = default;
+	TEOSGlobalCallback(TWeakPtr<OwningType> InOwner)
+		: FCallbackBase()
+		, Owner(InOwner)
+	{
+	}
 	virtual ~TEOSGlobalCallback() = default;
 
 
@@ -48,6 +48,9 @@ public:
 	}
 
 private:
+	/** The object that needs to be checked for lifetime before calling the callback */
+	TWeakPtr<OwningType> Owner;
+
 	static void EOS_CALL CallbackImpl(const CallbackType* Data)
 	{
 		check(IsInGameThread());
@@ -55,13 +58,11 @@ private:
 		TEOSGlobalCallback* CallbackThis = (TEOSGlobalCallback*)Data->ClientData;
 		check(CallbackThis);
 
-		if (FCallbackBase::ShouldCancelAllCallbacks())
+		if (CallbackThis->Owner.IsValid())
 		{
-			return;
+			check(CallbackThis->CallbackLambda);
+			CallbackThis->CallbackLambda(Data);
 		}
-
-		check(CallbackThis->CallbackLambda);
-		CallbackThis->CallbackLambda(Data);
 	}
 };
 
diff --git a/Engine/Plugins/Online/OnlineSubsystemEOS/Source/OnlineSubsystemEOS/Private/OnlineAchievementsEOS.cpp b/Engine/Plugins/Online/OnlineSubsystemEOS/Source/OnlineSubsystemEOS/Private/OnlineAchievementsEOS.cpp
--- a/Engine/Plugins/Online/OnlineSubsystemEOS/Source/OnlineSubsystemEOS/Private/OnlineAchievementsEOS.cpp
+++ b/Engine/Plugins/Online/OnlineSubsystemEOS/Source/OnlineSubsystemEOS/Private/OnlineAchievementsEOS.cpp
@@ -27,7 +27,7 @@ void FOnlineAchievementsEOS::WriteAchievements(const FUniqueNetId& PlayerId, FOn
 	Delegate.ExecuteIfBound(PlayerId, true);
 }
 
-typedef TEOSCallback<EOS_Achievements_OnQueryPlayerAchievementsCompleteCallback, EOS_Achievements_OnQueryPlayerAchievementsCompleteCallbackInfo> FQueryProgressCallback;
+typedef TEOSCallback<EOS_Achievements_OnQueryPlayerAchievementsCompleteCallback, EOS_Achievements_OnQueryPlayerAchievementsCompleteCallbackInfo, FOnlineAchievementsEOS> FQueryProgressCallback;
 
 void FOnlineAchievementsEOS::QueryAchievements(const FUniqueNetId& PlayerId, const FOnQueryAchievementsCompleteDelegate& Delegate)
 {
@@ -48,7 +48,7 @@ void FOnlineAchievementsEOS::QueryAchievements(const FUniqueNetId& PlayerId, con
 	Options.UserId = EOSSubsystem->UserManager->GetLocalProductUserId(LocalUserId);
 #endif
 
-	FQueryProgressCallback* CallbackObj = new FQueryProgressCallback();
+	FQueryProgressCallback* CallbackObj = new FQueryProgressCallback(FOnlineAchievementsEOSWeakPtr(AsShared()));
 	CallbackObj->CallbackLambda = [this, LambdaPlayerId = PlayerId.AsShared(), OnComplete = FOnQueryAchievementsCompleteDelegate(Delegate)](const EOS_Achievements_OnQueryPlayerAchievementsCompleteCallbackInfo* Data)
 	{
 		bool bWasSuccessful = Data->ResultCode == EOS_EResult::EOS_Success;
@@ -109,7 +109,7 @@ void FOnlineAchievementsEOS::QueryAchievements(const FUniqueNetId& PlayerId, con
 	EOS_Achievements_QueryPlayerAchievements(EOSSubsystem->AchievementsHandle, &Options, CallbackObj, CallbackObj->GetCallbackPtr());
 }
 
-typedef TEOSCallback<EOS_Achievements_OnQueryDefinitionsCompleteCallback, EOS_Achievements_OnQueryDefinitionsCompleteCallbackInfo> FQueryDefinitionsCallback;
+typedef TEOSCallback<EOS_Achievements_OnQueryDefinitionsCompleteCallback, EOS_Achievements_OnQueryDefinitionsCompleteCallbackInfo, FOnlineAchievementsEOS> FQueryDefinitionsCallback;
 
 void FOnlineAchievementsEOS::QueryAchievementDescriptions(const FUniqueNetId& PlayerId, const FOnQueryAchievementsCompleteDelegate& Delegate)
 {
@@ -132,7 +132,7 @@ void FOnlineAchievementsEOS::QueryAchievementDescriptions(const FUniqueNetId& Pl
 	Options.ApiVersion = EOS_ACHIEVEMENTS_QUERYDEFINITIONS_API_LATEST;
 	Options.LocalUserId = EOSSubsystem->UserManager->GetLocalProductUserId(LocalUserId);
 
-	FQueryDefinitionsCallback* CallbackObj = new FQueryDefinitionsCallback();
+	FQueryDefinitionsCallback* CallbackObj = new FQueryDefinitionsCallback(FOnlineAchievementsEOSWeakPtr(AsShared()));
 	CallbackObj->CallbackLambda = [this, LambdaPlayerId = PlayerId.AsShared(), OnComplete = FOnQueryAchievementsCompleteDelegate(Delegate)](const EOS_Achievements_OnQueryDefinitionsCompleteCallbackInfo* Data)
 	{
 		bool bWasSuccessful = Data->ResultCode == EOS_EResult::EOS_Success;
diff --git a/Engine/Plugins/Online/OnlineSubsystemEOS/Source/OnlineSubsystemEOS/Private/OnlineAchievementsEOS.h b/Engine/Plugins/Online/OnlineSubsystemEOS/Source/OnlineSubsystemEOS/Private/OnlineAchievementsEOS.h
--- a/Engine/Plugins/Online/OnlineSubsystemEOS/Source/OnlineSubsystemEOS/Private/OnlineAchievementsEOS.h
+++ b/Engine/Plugins/Online/OnlineSubsystemEOS/Source/OnlineSubsystemEOS/Private/OnlineAchievementsEOS.h
@@ -16,8 +16,9 @@ class FOnlineSubsystemEOS;
 /**
  * Interface for interacting with EOS achievements
  */
-class FOnlineAchievementsEOS :
-	public IOnlineAchievements
+class FOnlineAchievementsEOS
+	: public IOnlineAchievements
+	, public TSharedFromThis<FOnlineAchievementsEOS, ESPMode::ThreadSafe>
 {
 public:
 	FOnlineAchievementsEOS() = delete;
@@ -53,5 +54,6 @@ private:
 };
 
 typedef TSharedPtr<FOnlineAchievementsEOS, ESPMode::ThreadSafe> FOnlineAchievementsEOSPtr;
+typedef TWeakPtr<FOnlineAchievementsEOS, ESPMode::ThreadSafe> FOnlineAchievementsEOSWeakPtr;
 
 #endif
diff --git a/Engine/Plugins/Online/OnlineSubsystemEOS/Source/OnlineSubsystemEOS/Private/OnlineLeaderboardsEOS.cpp b/Engine/Plugins/Online/OnlineSubsystemEOS/Source/OnlineSubsystemEOS/Private/OnlineLeaderboardsEOS.cpp
--- a/Engine/Plugins/Online/OnlineSubsystemEOS/Source/OnlineSubsystemEOS/Private/OnlineLeaderboardsEOS.cpp
+++ b/Engine/Plugins/Online/OnlineSubsystemEOS/Source/OnlineSubsystemEOS/Private/OnlineLeaderboardsEOS.cpp
@@ -73,7 +73,7 @@ struct FQueryLeaderboardForUsersContext
 	}
 };
 
-typedef TEOSCallback<EOS_Leaderboards_OnQueryLeaderboardUserScoresCompleteCallback, EOS_Leaderboards_OnQueryLeaderboardUserScoresCompleteCallbackInfo> FQueryLeaderboardForUsersCallback;
+typedef TEOSCallback<EOS_Leaderboards_OnQueryLeaderboardUserScoresCompleteCallback, EOS_Leaderboards_OnQueryLeaderboardUserScoresCompleteCallbackInfo, FOnlineLeaderboardsEOS> FQueryLeaderboardForUsersCallback;
 
 bool FOnlineLeaderboardsEOS::ReadLeaderboards(const TArray<FUniqueNetIdRef>& Players, FOnlineLeaderboardReadRef& ReadObject)
 {
@@ -115,7 +115,7 @@ bool FOnlineLeaderboardsEOS::ReadLeaderboards(const TArray<FUniqueNetIdRef>& Pla
 
 	TSharedPtr<FQueryLeaderboardForUsersContext> QueryContext = MakeShared<FQueryLeaderboardForUsersContext>(Players, ReadObject);
 
-	FQueryLeaderboardForUsersCallback* CallbackObj = new FQueryLeaderboardForUsersCallback();
+	FQueryLeaderboardForUsersCallback* CallbackObj = new FQueryLeaderboardForUsersCallback(FOnlineLeaderboardsEOSWeakPtr(AsShared()));
 	CallbackObj->CallbackLambda = [this, QueryContext](const EOS_Leaderboards_OnQueryLeaderboardUserScoresCompleteCallbackInfo* Data)
 	{
 		bool bWasSuccessful = Data->ResultCode == EOS_EResult::EOS_Success;
@@ -226,7 +226,7 @@ bool FOnlineLeaderboardsEOS::ReadLeaderboardsForFriends(int32 LocalUserNum, FOnl
 	return ReadLeaderboards(Players, ReadObject);
 }
 
-typedef TEOSCallback<EOS_Leaderboards_OnQueryLeaderboardRanksCompleteCallback, EOS_Leaderboards_OnQueryLeaderboardRanksCompleteCallbackInfo> FQueryLeaderboardCallback;
+typedef TEOSCallback<EOS_Leaderboards_OnQueryLeaderboardRanksCompleteCallback, EOS_Leaderboards_OnQueryLeaderboardRanksCompleteCallbackInfo, FOnlineLeaderboardsEOS> FQueryLeaderboardCallback;
 
 bool FOnlineLeaderboardsEOS::ReadLeaderboardsAroundRank(int32 Rank, uint32 Range, FOnlineLeaderboardReadRef& ReadObject)
 {
@@ -251,7 +251,7 @@ bool FOnlineLeaderboardsEOS::ReadLeaderboardsAroundRank(int32 Rank, uint32 Range
 	Options.LocalUserId = EOSSubsystem->UserManager->GetLocalProductUserId(0);
 	FCStringAnsi::Strncpy(LeaderboardId, TCHAR_TO_UTF8(*ReadObject->LeaderboardName.ToString()), EOS_OSS_STRING_BUFFER_LENGTH);
 
-	FQueryLeaderboardCallback* CallbackObj = new FQueryLeaderboardCallback();
+	FQueryLeaderboardCallback* CallbackObj = new FQueryLeaderboardCallback(FOnlineLeaderboardsEOSWeakPtr(AsShared()));
 	FOnlineLeaderboardReadRef LambdaReadObject = ReadObject;
 	CallbackObj->CallbackLambda = [this, LambdaReadObject, StartIndex, EndIndex](const EOS_Leaderboards_OnQueryLeaderboardRanksCompleteCallbackInfo* Data)
 	{
diff --git a/Engine/Plugins/Online/OnlineSubsystemEOS/Source/OnlineSubsystemEOS/Private/OnlineLeaderboardsEOS.h b/Engine/Plugins/Online/OnlineSubsystemEOS/Source/OnlineSubsystemEOS/Private/OnlineLeaderboardsEOS.h
--- a/Engine/Plugins/Online/OnlineSubsystemEOS/Source/OnlineSubsystemEOS/Private/OnlineLeaderboardsEOS.h
+++ b/Engine/Plugins/Online/OnlineSubsystemEOS/Source/OnlineSubsystemEOS/Private/OnlineLeaderboardsEOS.h
@@ -16,8 +16,9 @@ class FOnlineSubsystemEOS;
 /**
  * Interface for interacting with EOS stats
  */
-class FOnlineLeaderboardsEOS :
-	public IOnlineLeaderboards
+class FOnlineLeaderboardsEOS
+	: public IOnlineLeaderboards
+	, public TSharedFromThis<FOnlineLeaderboardsEOS, ESPMode::ThreadSafe>
 {
 public:
 	FOnlineLeaderboardsEOS() = delete;
@@ -47,5 +48,6 @@ private:
 };
 
 typedef TSharedPtr<FOnlineLeaderboardsEOS, ESPMode::ThreadSafe> FOnlineLeaderboardsEOSPtr;
+typedef TWeakPtr<FOnlineLeaderboardsEOS, ESPMode::ThreadSafe> FOnlineLeaderboardsEOSWeakPtr;
 
 #endif
diff --git a/Engine/Plugins/Online/OnlineSubsystemEOS/Source/OnlineSubsystemEOS/Private/OnlineSessionEOS.cpp b/Engine/Plugins/Online/OnlineSubsystemEOS/Source/OnlineSubsystemEOS/Private/OnlineSessionEOS.cpp
--- a/Engine/Plugins/Online/OnlineSubsystemEOS/Source/OnlineSubsystemEOS/Private/OnlineSessionEOS.cpp
+++ b/Engine/Plugins/Online/OnlineSubsystemEOS/Source/OnlineSubsystemEOS/Private/OnlineSessionEOS.cpp
@@ -412,25 +412,25 @@ void FOnlineSessionInfoEOS::InitLAN(FOnlineSubsystemEOS* Subsystem)
 	SessionId = FUniqueNetIdEOS::Create(OwnerGuid.ToString());
 }
 
-typedef TEOSGlobalCallback<EOS_Sessions_OnSessionInviteReceivedCallback, EOS_Sessions_SessionInviteReceivedCallbackInfo> FSessionInviteReceivedCallback;
-typedef TEOSGlobalCallback<EOS_Sessions_OnSessionInviteAcceptedCallback, EOS_Sessions_SessionInviteAcceptedCallbackInfo> FSessionInviteAcceptedCallback;
+typedef TEOSGlobalCallback<EOS_Sessions_OnSessionInviteReceivedCallback, EOS_Sessions_SessionInviteReceivedCallbackInfo, FOnlineSessionEOS> FSessionInviteReceivedCallback;
+typedef TEOSGlobalCallback<EOS_Sessions_OnSessionInviteAcceptedCallback, EOS_Sessions_SessionInviteAcceptedCallbackInfo, FOnlineSessionEOS> FSessionInviteAcceptedCallback;
 
 // Lobby session callbacks
-typedef TEOSCallback<EOS_Lobby_OnCreateLobbyCallback, EOS_Lobby_CreateLobbyCallbackInfo> FLobbyCreatedCallback;
-typedef TEOSCallback<EOS_Lobby_OnUpdateLobbyCallback, EOS_Lobby_UpdateLobbyCallbackInfo> FLobbyUpdatedCallback;
-typedef TEOSCallback<EOS_Lobby_OnJoinLobbyCallback, EOS_Lobby_JoinLobbyCallbackInfo> FLobbyJoinedCallback;
-typedef TEOSCallback<EOS_Lobby_OnLeaveLobbyCallback, EOS_Lobby_LeaveLobbyCallbackInfo> FLobbyLeftCallback;
-typedef TEOSCallback<EOS_Lobby_OnDestroyLobbyCallback, EOS_Lobby_DestroyLobbyCallbackInfo> FLobbyDestroyedCallback;
-typedef TEOSCallback<EOS_Lobby_OnSendInviteCallback, EOS_Lobby_SendInviteCallbackInfo> FLobbySendInviteCallback;
-typedef TEOSCallback<EOS_Lobby_OnKickMemberCallback, EOS_Lobby_KickMemberCallbackInfo> FLobbyRemovePlayerCallback;
-typedef TEOSCallback<EOS_LobbySearch_OnFindCallback, EOS_LobbySearch_FindCallbackInfo> FLobbySearchFindCallback;
+typedef TEOSCallback<EOS_Lobby_OnCreateLobbyCallback, EOS_Lobby_CreateLobbyCallbackInfo, FOnlineSessionEOS> FLobbyCreatedCallback;
+typedef TEOSCallback<EOS_Lobby_OnUpdateLobbyCallback, EOS_Lobby_UpdateLobbyCallbackInfo, FOnlineSessionEOS> FLobbyUpdatedCallback;
+typedef TEOSCallback<EOS_Lobby_OnJoinLobbyCallback, EOS_Lobby_JoinLobbyCallbackInfo, FOnlineSessionEOS> FLobbyJoinedCallback;
+typedef TEOSCallback<EOS_Lobby_OnLeaveLobbyCallback, EOS_Lobby_LeaveLobbyCallbackInfo, FOnlineSessionEOS> FLobbyLeftCallback;
+typedef TEOSCallback<EOS_Lobby_OnDestroyLobbyCallback, EOS_Lobby_DestroyLobbyCallbackInfo, FOnlineSessionEOS> FLobbyDestroyedCallback;
+typedef TEOSCallback<EOS_Lobby_OnSendInviteCallback, EOS_Lobby_SendInviteCallbackInfo, FOnlineSessionEOS> FLobbySendInviteCallback;
+typedef TEOSCallback<EOS_Lobby_OnKickMemberCallback, EOS_Lobby_KickMemberCallbackInfo, FOnlineSessionEOS> FLobbyRemovePlayerCallback;
+typedef TEOSCallback<EOS_LobbySearch_OnFindCallback, EOS_LobbySearch_FindCallbackInfo, FOnlineSessionEOS> FLobbySearchFindCallback;
 
 // Lobby notification callbacks
-typedef TEOSGlobalCallback<EOS_Lobby_OnLobbyUpdateReceivedCallback, EOS_Lobby_LobbyUpdateReceivedCallbackInfo> FLobbyUpdateReceivedCallback;
-typedef TEOSGlobalCallback<EOS_Lobby_OnLobbyMemberUpdateReceivedCallback, EOS_Lobby_LobbyMemberUpdateReceivedCallbackInfo> FLobbyMemberUpdateReceivedCallback;
-typedef TEOSGlobalCallback<EOS_Lobby_OnLobbyMemberStatusReceivedCallback, EOS_Lobby_LobbyMemberStatusReceivedCallbackInfo> FLobbyMemberStatusReceivedCallback;
-typedef TEOSGlobalCallback<EOS_Lobby_OnLobbyInviteAcceptedCallback, EOS_Lobby_LobbyInviteAcceptedCallbackInfo> FLobbyInviteAcceptedCallback;
-typedef TEOSGlobalCallback<EOS_Lobby_OnJoinLobbyAcceptedCallback, EOS_Lobby_JoinLobbyAcceptedCallbackInfo> FJoinLobbyAcceptedCallback;
+typedef TEOSGlobalCallback<EOS_Lobby_OnLobbyUpdateReceivedCallback, EOS_Lobby_LobbyUpdateReceivedCallbackInfo, FOnlineSessionEOS> FLobbyUpdateReceivedCallback;
+typedef TEOSGlobalCallback<EOS_Lobby_OnLobbyMemberUpdateReceivedCallback, EOS_Lobby_LobbyMemberUpdateReceivedCallbackInfo, FOnlineSessionEOS> FLobbyMemberUpdateReceivedCallback;
+typedef TEOSGlobalCallback<EOS_Lobby_OnLobbyMemberStatusReceivedCallback, EOS_Lobby_LobbyMemberStatusReceivedCallbackInfo, FOnlineSessionEOS> FLobbyMemberStatusReceivedCallback;
+typedef TEOSGlobalCallback<EOS_Lobby_OnLobbyInviteAcceptedCallback, EOS_Lobby_LobbyInviteAcceptedCallbackInfo, FOnlineSessionEOS> FLobbyInviteAcceptedCallback;
+typedef TEOSGlobalCallback<EOS_Lobby_OnJoinLobbyAcceptedCallback, EOS_Lobby_JoinLobbyAcceptedCallbackInfo, FOnlineSessionEOS> FJoinLobbyAcceptedCallback;
 
 FOnlineSessionEOS::~FOnlineSessionEOS()
 {
@@ -455,7 +455,7 @@ void FOnlineSessionEOS::Init(const FString& InBucketId)
 	FCStringAnsi::Strncpy(BucketIdAnsi, TCHAR_TO_UTF8(*InBucketId), EOS_OSS_STRING_BUFFER_LENGTH);
 
 	// Register for session invite notifications
-	FSessionInviteAcceptedCallback* SessionInviteAcceptedCallbackObj = new FSessionInviteAcceptedCallback();
+	FSessionInviteAcceptedCallback* SessionInviteAcceptedCallbackObj = new FSessionInviteAcceptedCallback(FOnlineSessionEOSWeakPtr(AsShared()));
 	SessionInviteAcceptedCallback = SessionInviteAcceptedCallbackObj;
 	SessionInviteAcceptedCallbackObj->CallbackLambda = [this](const EOS_Sessions_SessionInviteAcceptedCallbackInfo* Data)
 	{
@@ -538,7 +538,7 @@ void FOnlineSessionEOS::RegisterLobbyNotifications()
 	EOS_Lobby_AddNotifyLobbyUpdateReceivedOptions AddNotifyLobbyUpdateReceivedOptions = { 0 };
 	AddNotifyLobbyUpdateReceivedOptions.ApiVersion = EOS_LOBBY_ADDNOTIFYLOBBYUPDATERECEIVED_API_LATEST;
 
-	FLobbyUpdateReceivedCallback* LobbyUpdateReceivedCallbackObj = new FLobbyUpdateReceivedCallback();
+	FLobbyUpdateReceivedCallback* LobbyUpdateReceivedCallbackObj = new FLobbyUpdateReceivedCallback(FOnlineSessionEOSWeakPtr(AsShared()));
 	LobbyUpdateReceivedCallback = LobbyUpdateReceivedCallbackObj;
 	LobbyUpdateReceivedCallbackObj->CallbackLambda = [this](const EOS_Lobby_LobbyUpdateReceivedCallbackInfo* Data)
 	{
@@ -551,7 +551,7 @@ void FOnlineSessionEOS::RegisterLobbyNotifications()
 	EOS_Lobby_AddNotifyLobbyMemberUpdateReceivedOptions AddNotifyLobbyMemberUpdateReceivedOptions = { 0 };
 	AddNotifyLobbyMemberUpdateReceivedOptions.ApiVersion = EOS_LOBBY_ADDNOTIFYLOBBYMEMBERUPDATERECEIVED_API_LATEST;
 
-	FLobbyMemberUpdateReceivedCallback* LobbyMemberUpdateReceivedCallbackObj = new FLobbyMemberUpdateReceivedCallback();
+	FLobbyMemberUpdateReceivedCallback* LobbyMemberUpdateReceivedCallbackObj = new FLobbyMemberUpdateReceivedCallback(FOnlineSessionEOSWeakPtr(AsShared()));
 	LobbyMemberUpdateReceivedCallback = LobbyMemberUpdateReceivedCallbackObj;
 	LobbyMemberUpdateReceivedCallbackObj->CallbackLambda = [this](const EOS_Lobby_LobbyMemberUpdateReceivedCallbackInfo* Data)
 	{
@@ -564,7 +564,7 @@ void FOnlineSessionEOS::RegisterLobbyNotifications()
 	EOS_Lobby_AddNotifyLobbyMemberStatusReceivedOptions AddNotifyLobbyMemberStatusReceivedOptions = { 0 };
 	AddNotifyLobbyMemberStatusReceivedOptions.ApiVersion = EOS_LOBBY_ADDNOTIFYLOBBYMEMBERSTATUSRECEIVED_API_LATEST;
 
-	FLobbyMemberStatusReceivedCallback* LobbyMemberStatusReceivedCallbackObj = new FLobbyMemberStatusReceivedCallback();
+	FLobbyMemberStatusReceivedCallback* LobbyMemberStatusReceivedCallbackObj = new FLobbyMemberStatusReceivedCallback(FOnlineSessionEOSWeakPtr(AsShared()));
 	LobbyMemberStatusReceivedCallback = LobbyMemberStatusReceivedCallbackObj;
 	LobbyMemberStatusReceivedCallbackObj->CallbackLambda = [this](const EOS_Lobby_LobbyMemberStatusReceivedCallbackInfo* Data)
 	{
@@ -579,7 +579,7 @@ void FOnlineSessionEOS::RegisterLobbyNotifications()
 	EOS_Lobby_AddNotifyLobbyInviteAcceptedOptions AddNotifyLobbyInviteAcceptedOptions = { 0 };
 	AddNotifyLobbyInviteAcceptedOptions.ApiVersion = EOS_LOBBY_ADDNOTIFYLOBBYINVITEACCEPTED_API_LATEST;
 
-	FLobbyInviteAcceptedCallback* LobbyInviteAcceptedCallbackObj = new FLobbyInviteAcceptedCallback();
+	FLobbyInviteAcceptedCallback* LobbyInviteAcceptedCallbackObj = new FLobbyInviteAcceptedCallback(FOnlineSessionEOSWeakPtr(AsShared()));
 	LobbyInviteAcceptedCallback = LobbyInviteAcceptedCallbackObj;
 	LobbyInviteAcceptedCallbackObj->CallbackLambda = [this](const EOS_Lobby_LobbyInviteAcceptedCallbackInfo* Data)
 	{
@@ -592,7 +592,7 @@ void FOnlineSessionEOS::RegisterLobbyNotifications()
 	EOS_Lobby_AddNotifyJoinLobbyAcceptedOptions AddNotifyJoinLobbyAcceptedOptions = { 0 };
 	AddNotifyJoinLobbyAcceptedOptions.ApiVersion = EOS_LOBBY_ADDNOTIFYJOINLOBBYACCEPTED_API_LATEST;
 
-	FJoinLobbyAcceptedCallback* JoinLobbyAcceptedCallbackObj = new FJoinLobbyAcceptedCallback();
+	FJoinLobbyAcceptedCallback* JoinLobbyAcceptedCallbackObj = new FJoinLobbyAcceptedCallback(FOnlineSessionEOSWeakPtr(AsShared()));
 	JoinLobbyAcceptedCallback = JoinLobbyAcceptedCallbackObj;
 	JoinLobbyAcceptedCallbackObj->CallbackLambda = [this](const EOS_Lobby_JoinLobbyAcceptedCallbackInfo* Data)
 	{
@@ -1212,7 +1212,7 @@ uint32 FOnlineSessionEOS::CreateEOSSession(int32 HostingPlayerNum, FNamedOnlineS
 
 	FName SessionName = Session->SessionName;
 
-	FUpdateSessionCallback* CallbackObj = new FUpdateSessionCallback();
+	FUpdateSessionCallback* CallbackObj = new FUpdateSessionCallback(FOnlineSessionEOSWeakPtr(AsShared()));
 	CallbackObj->CallbackLambda = [this, SessionName](const EOS_Sessions_UpdateSessionCallbackInfo* Data)
 	{
 		bool bWasSuccessful = false;
@@ -1336,14 +1336,14 @@ struct FSessionStartOptions :
 	}
 };
 
-typedef TEOSCallback<EOS_Sessions_OnStartSessionCallback, EOS_Sessions_StartSessionCallbackInfo> FStartSessionCallback;
+typedef TEOSCallback<EOS_Sessions_OnStartSessionCallback, EOS_Sessions_StartSessionCallbackInfo, FOnlineSessionEOS> FStartSessionCallback;
 
 uint32 FOnlineSessionEOS::StartEOSSession(FNamedOnlineSession* Session)
 {
 	Session->SessionState = EOnlineSessionState::Starting;
 
 	FSessionStartOptions Options(TCHAR_TO_UTF8(*Session->SessionName.ToString()));
-	FStartSessionCallback* CallbackObj = new FStartSessionCallback();
+	FStartSessionCallback* CallbackObj = new FStartSessionCallback(FOnlineSessionEOSWeakPtr(AsShared()));
 	CallbackObj->CallbackLambda = [this, Session](const EOS_Sessions_StartSessionCallbackInfo* Data)
 	{
 		Session->SessionState = EOnlineSessionState::InProgress;
@@ -1444,7 +1444,7 @@ uint32 FOnlineSessionEOS::UpdateEOSSession(FNamedOnlineSession* Session)
 		return ONLINE_FAIL;
 	}
 
-	FUpdateSessionCallback* CallbackObj = new FUpdateSessionCallback();
+	FUpdateSessionCallback* CallbackObj = new FUpdateSessionCallback(FOnlineSessionEOSWeakPtr(AsShared()));
 	CallbackObj->CallbackLambda = [this, Session](const EOS_Sessions_UpdateSessionCallbackInfo* Data)
 	{
 		bool bWasSuccessful = Data->ResultCode == EOS_EResult::EOS_Success || Data->ResultCode == EOS_EResult::EOS_Sessions_OutOfSync;
@@ -1537,7 +1537,7 @@ struct FSessionEndOptions :
 	}
 };
 
-typedef TEOSCallback<EOS_Sessions_OnEndSessionCallback, EOS_Sessions_EndSessionCallbackInfo> FEndSessionCallback;
+typedef TEOSCallback<EOS_Sessions_OnEndSessionCallback, EOS_Sessions_EndSessionCallbackInfo, FOnlineSessionEOS> FEndSessionCallback;
 
 uint32 FOnlineSessionEOS::EndEOSSession(FNamedOnlineSession* Session)
 {
@@ -1547,7 +1547,7 @@ uint32 FOnlineSessionEOS::EndEOSSession(FNamedOnlineSession* Session)
 	Session->SessionState = EOnlineSessionState::Ending;
 
 	FSessionEndOptions Options(TCHAR_TO_UTF8(*Session->SessionName.ToString()));
-	FEndSessionCallback* CallbackObj = new FEndSessionCallback();
+	FEndSessionCallback* CallbackObj = new FEndSessionCallback(FOnlineSessionEOSWeakPtr(AsShared()));
 	CallbackObj->CallbackLambda = [this, Session](const EOS_Sessions_EndSessionCallbackInfo* Data)
 	{
 		Session->SessionState = EOnlineSessionState::Ended;
@@ -1681,14 +1681,14 @@ struct FSessionDestroyOptions :
 	}
 };
 
-typedef TEOSCallback<EOS_Sessions_OnDestroySessionCallback, EOS_Sessions_DestroySessionCallbackInfo> FDestroySessionCallback;
+typedef TEOSCallback<EOS_Sessions_OnDestroySessionCallback, EOS_Sessions_DestroySessionCallbackInfo, FOnlineSessionEOS> FDestroySessionCallback;
 
 uint32 FOnlineSessionEOS::DestroyEOSSession(FNamedOnlineSession* Session, const FOnDestroySessionCompleteDelegate& CompletionDelegate)
 {
 	Session->SessionState = EOnlineSessionState::Destroying;
 
 	FSessionDestroyOptions Options(TCHAR_TO_UTF8(*Session->SessionName.ToString()));
-	FDestroySessionCallback* CallbackObj = new FDestroySessionCallback();
+	FDestroySessionCallback* CallbackObj = new FDestroySessionCallback(FOnlineSessionEOSWeakPtr(AsShared()));
 	CallbackObj->CallbackLambda = [this, Session](const EOS_Sessions_DestroySessionCallbackInfo* Data)
 	{
 		EndSessionAnalytics();
@@ -2006,7 +2006,7 @@ void FOnlineSessionEOS::AddSearchResult(EOS_HSessionDetails SessionHandle, const
 	}
 }
 
-typedef TEOSCallback<EOS_SessionSearch_OnFindCallback, EOS_SessionSearch_FindCallbackInfo> FFindSessionsCallback;
+typedef TEOSCallback<EOS_SessionSearch_OnFindCallback, EOS_SessionSearch_FindCallbackInfo, FOnlineSessionEOS> FFindSessionsCallback;
 
 uint32 FOnlineSessionEOS::FindEOSSession(int32 SearchingPlayerNum, const TSharedRef<FOnlineSessionSearch>& SearchSettings)
 {
@@ -2049,7 +2049,7 @@ uint32 FOnlineSessionEOS::FindEOSSession(int32 SearchingPlayerNum, const TShared
 		AddSearchAttribute(SearchHandle, &Attribute, ToEOSSearchOp(SearchParam.ComparisonOp));
 	}
 
-	FFindSessionsCallback* CallbackObj = new FFindSessionsCallback();
+	FFindSessionsCallback* CallbackObj = new FFindSessionsCallback(FOnlineSessionEOSWeakPtr(AsShared()));
 	CallbackObj->CallbackLambda = [this, SearchSettings](const EOS_SessionSearch_FindCallbackInfo* Data)
 	{
 		bool bWasSuccessful = Data->ResultCode == EOS_EResult::EOS_Success;
@@ -2122,7 +2122,7 @@ void FOnlineSessionEOS::FindEOSSessionById(int32 LocalUserNum, const FUniqueNetI
 	// Store our search handle for use/cleanup later
 	CurrentSearchHandle = MakeShareable(new FSessionSearchEOS(SearchHandle));
 
-	FFindSessionsCallback* CallbackObj = new FFindSessionsCallback();
+	FFindSessionsCallback* CallbackObj = new FFindSessionsCallback(FOnlineSessionEOSWeakPtr(AsShared()));
 	CallbackObj->CallbackLambda = [this, LocalUserNum, OnComplete = FOnSingleSessionResultCompleteDelegate(CompletionDelegate)](const EOS_SessionSearch_FindCallbackInfo* Data)
 	{
 		TSharedRef<FOnlineSessionSearch> LocalSessionSearch = MakeShareable(new FOnlineSessionSearch());
@@ -2326,7 +2326,7 @@ struct FJoinSessionOptions :
 	}
 };
 
-typedef TEOSCallback<EOS_Sessions_OnJoinSessionCallback, EOS_Sessions_JoinSessionCallbackInfo> FJoinSessionCallback;
+typedef TEOSCallback<EOS_Sessions_OnJoinSessionCallback, EOS_Sessions_JoinSessionCallbackInfo, FOnlineSessionEOS> FJoinSessionCallback;
 
 uint32 FOnlineSessionEOS::JoinEOSSession(int32 PlayerNum, FNamedOnlineSession* Session, const FOnlineSession* SearchSession)
 {
@@ -2356,7 +2356,7 @@ uint32 FOnlineSessionEOS::JoinEOSSession(int32 PlayerNum, FNamedOnlineSession* S
 
 	FName SessionName = Session->SessionName;
 
-	FJoinSessionCallback* CallbackObj = new FJoinSessionCallback();
+	FJoinSessionCallback* CallbackObj = new FJoinSessionCallback(FOnlineSessionEOSWeakPtr(AsShared()));
 	CallbackObj->CallbackLambda = [this, SessionName](const EOS_Sessions_JoinSessionCallbackInfo* Data)
 	{
 		bool bWasSuccessful = false;
@@ -2484,7 +2484,7 @@ struct FSendSessionInviteOptions :
 	}
 };
 
-typedef TEOSCallback<EOS_Sessions_OnSendInviteCallback, EOS_Sessions_SendInviteCallbackInfo> FSendSessionInviteCallback;
+typedef TEOSCallback<EOS_Sessions_OnSendInviteCallback, EOS_Sessions_SendInviteCallbackInfo, FOnlineSessionEOS> FSendSessionInviteCallback;
 
 bool FOnlineSessionEOS::SendSessionInvite(FName SessionName, EOS_ProductUserId SenderId, EOS_ProductUserId ReceiverId)
 {
@@ -2519,7 +2519,7 @@ bool FOnlineSessionEOS::SendLobbyInvite(FName SessionName, EOS_ProductUserId Sen
 	SendInviteOptions.LocalUserId = SenderId;
 	SendInviteOptions.TargetUserId = ReceiverId;
 	
-	FLobbySendInviteCallback* CallbackObj = new FLobbySendInviteCallback();
+	FLobbySendInviteCallback* CallbackObj = new FLobbySendInviteCallback(FOnlineSessionEOSWeakPtr(AsShared()));
 	LobbySendInviteCallback = CallbackObj;
 	CallbackObj->CallbackLambda = [this](const EOS_Lobby_SendInviteCallbackInfo* Data)
 	{
@@ -2544,7 +2544,7 @@ bool FOnlineSessionEOS::SendEOSSessionInvite(FName SessionName, EOS_ProductUserI
 	Options.LocalUserId = SenderId;
 	Options.TargetUserId = ReceiverId;
 
-	FSendSessionInviteCallback* CallbackObj = new FSendSessionInviteCallback();
+	FSendSessionInviteCallback* CallbackObj = new FSendSessionInviteCallback(FOnlineSessionEOSWeakPtr(AsShared()));
 	CallbackObj->CallbackLambda = [this, SessionName](const EOS_Sessions_SendInviteCallbackInfo* Data)
 	{
 		bool bWasSuccessful = Data->ResultCode == EOS_EResult::EOS_Success;
@@ -2730,7 +2730,7 @@ bool FOnlineSessionEOS::RegisterPlayer(FName SessionName, const FUniqueNetId& Pl
 	return RegisterPlayers(SessionName, Players, bWasInvited);
 }
 
-typedef TEOSCallback<EOS_Sessions_OnRegisterPlayersCallback, EOS_Sessions_RegisterPlayersCallbackInfo> FRegisterPlayersCallback;
+typedef TEOSCallback<EOS_Sessions_OnRegisterPlayersCallback, EOS_Sessions_RegisterPlayersCallbackInfo, FOnlineSessionEOS> FRegisterPlayersCallback;
 
 bool FOnlineSessionEOS::RegisterPlayers(FName SessionName, const TArray< FUniqueNetIdRef >& Players, bool bWasInvited)
 {
@@ -2784,7 +2784,7 @@ bool FOnlineSessionEOS::RegisterPlayers(FName SessionName, const TArray< FUnique
 			const FTCHARToUTF8 Utf8SessionName(*SessionName.ToString());
 			Options.SessionName = Utf8SessionName.Get();
 
-			FRegisterPlayersCallback* CallbackObj = new FRegisterPlayersCallback();
+			FRegisterPlayersCallback* CallbackObj = new FRegisterPlayersCallback(FOnlineSessionEOSWeakPtr(AsShared()));
 			CallbackObj->CallbackLambda = [this, SessionName, RegisteredPlayers = TArray<FUniqueNetIdRef>(Players)](const EOS_Sessions_RegisterPlayersCallbackInfo* Data)
 			{
 				bool bWasSuccessful = Data->ResultCode == EOS_EResult::EOS_Success || Data->ResultCode == EOS_EResult::EOS_NoChange;
@@ -2814,7 +2814,7 @@ bool FOnlineSessionEOS::UnregisterPlayer(FName SessionName, const FUniqueNetId&
 	return UnregisterPlayers(SessionName, Players);
 }
 
-typedef TEOSCallback<EOS_Sessions_OnUnregisterPlayersCallback, EOS_Sessions_UnregisterPlayersCallbackInfo> FUnregisterPlayersCallback;
+typedef TEOSCallback<EOS_Sessions_OnUnregisterPlayersCallback, EOS_Sessions_UnregisterPlayersCallbackInfo, FOnlineSessionEOS> FUnregisterPlayersCallback;
 
 bool FOnlineSessionEOS::UnregisterPlayers(FName SessionName, const TArray< FUniqueNetIdRef >& Players)
 {
@@ -2868,7 +2868,7 @@ bool FOnlineSessionEOS::UnregisterPlayers(FName SessionName, const TArray< FUniq
 			const FTCHARToUTF8 Utf8SessionName(*SessionName.ToString());
 			Options.SessionName = Utf8SessionName.Get();
 
-			FUnregisterPlayersCallback* CallbackObj = new FUnregisterPlayersCallback();
+			FUnregisterPlayersCallback* CallbackObj = new FUnregisterPlayersCallback(FOnlineSessionEOSWeakPtr(AsShared()));
 			CallbackObj->CallbackLambda = [this, SessionName, UnregisteredPlayers = TArray<FUniqueNetIdRef>(Players)](const EOS_Sessions_UnregisterPlayersCallbackInfo* Data)
 			{
 				bool bWasSuccessful = Data->ResultCode == EOS_EResult::EOS_Success || Data->ResultCode == EOS_EResult::EOS_NoChange;
@@ -3221,7 +3221,7 @@ void FOnlineSessionEOS::RemovePlayerFromSession(int32 LocalUserNum, FName Sessio
 		KickMemberOptions.LocalUserId = EOSSubsystem->UserManager->GetLocalProductUserId(LocalUserNum);
 		KickMemberOptions.TargetUserId = EOSSubsystem->UserManager->GetProductUserId(TargetPlayerId);
 
-		FLobbyRemovePlayerCallback* CallbackObj = new FLobbyRemovePlayerCallback();
+		FLobbyRemovePlayerCallback* CallbackObj = new FLobbyRemovePlayerCallback(FOnlineSessionEOSWeakPtr(AsShared()));
 		CallbackObj->CallbackLambda = [this](const EOS_Lobby_KickMemberCallbackInfo* Data)
 		{
 			if (Data->ResultCode == EOS_EResult::EOS_Success)
@@ -3336,7 +3336,7 @@ uint32 FOnlineSessionEOS::CreateLobbySession(int32 HostingPlayerNum, FNamedOnlin
 	If the data structure's ResultCode field indicates success, its LobbyId field contains the new lobby's ID value, which we will need to interact with the lobby further.*/
 
 	FName SessionName = Session->SessionName;
-	FLobbyCreatedCallback* CallbackObj = new FLobbyCreatedCallback();
+	FLobbyCreatedCallback* CallbackObj = new FLobbyCreatedCallback(FOnlineSessionEOSWeakPtr(AsShared()));
 	LobbyCreatedCallback = CallbackObj;
 	CallbackObj->CallbackLambda = [this, SessionName, LocalProductUserId, LocalUserNetId](const EOS_Lobby_CreateLobbyCallbackInfo* Data)
 	{
@@ -3417,7 +3417,7 @@ uint32 FOnlineSessionEOS::JoinLobbySession(int32 PlayerNum, FNamedOnlineSession*
 			FName SessionName = Session->SessionName;
 			FUniqueNetIdPtr LocalUserNetId = EOSSubsystem->UserManager->GetLocalUniqueNetIdEOS(PlayerNum);
 
-			FLobbyJoinedCallback* CallbackObj = new FLobbyJoinedCallback();
+			FLobbyJoinedCallback* CallbackObj = new FLobbyJoinedCallback(FOnlineSessionEOSWeakPtr(AsShared()));
 			LobbyJoinedCallback = CallbackObj;
 			CallbackObj->CallbackLambda = [this, SessionName, LocalUserNetId, LobbyDetailsHandle](const EOS_Lobby_JoinLobbyCallbackInfo* Data)
 			{
@@ -3678,7 +3678,7 @@ uint32 FOnlineSessionEOS::UpdateLobbySession(FNamedOnlineSession* Session)
 			UpdateLobbyOptions.LobbyModificationHandle = LobbyModificationHandle;
 
 			FName SessionName = Session->SessionName;
-			FLobbyUpdatedCallback* CallbackObj = new FLobbyUpdatedCallback();
+			FLobbyUpdatedCallback* CallbackObj = new FLobbyUpdatedCallback(FOnlineSessionEOSWeakPtr(AsShared()));
 			CallbackObj->CallbackLambda = [this, SessionName](const EOS_Lobby_UpdateLobbyCallbackInfo* Data)
 			{
 				FNamedOnlineSession* Session = GetNamedSession(SessionName);
@@ -3751,7 +3751,7 @@ uint32 FOnlineSessionEOS::DestroyLobbySession(FNamedOnlineSession* Session, cons
 		LeaveOptions.LocalUserId = EOSSubsystem->UserManager->GetLocalProductUserId(EOSSubsystem->UserManager->GetDefaultLocalUser()); // Maybe not split screen friendly
 
 		FName SessionName = Session->SessionName;
-		FLobbyLeftCallback* LeaveCallbackObj = new FLobbyLeftCallback();
+		FLobbyLeftCallback* LeaveCallbackObj = new FLobbyLeftCallback(FOnlineSessionEOSWeakPtr(AsShared()));
 		LobbyLeftCallback = LeaveCallbackObj;
 		LeaveCallbackObj->CallbackLambda = [this, SessionName, CompletionDelegate](const EOS_Lobby_LeaveLobbyCallbackInfo* Data)
 		{
@@ -3853,7 +3853,7 @@ void FOnlineSessionEOS::StartLobbySearch(int32 SearchingPlayerNum, EOS_HLobbySea
 	FindOptions.ApiVersion = EOS_LOBBYSEARCH_FIND_API_LATEST;
 	FindOptions.LocalUserId = EOSSubsystem->UserManager->GetLocalProductUserId(SearchingPlayerNum);
 
-	FLobbySearchFindCallback* CallbackObj = new FLobbySearchFindCallback();
+	FLobbySearchFindCallback* CallbackObj = new FLobbySearchFindCallback(FOnlineSessionEOSWeakPtr(AsShared()));
 	LobbySearchFindCallback = CallbackObj;
 	CallbackObj->CallbackLambda = [this, SearchingPlayerNum, LobbySearchHandle, SearchSettings, CompletionDelegate](const EOS_LobbySearch_FindCallbackInfo* Data)
 	{
diff --git a/Engine/Plugins/Online/OnlineSubsystemEOS/Source/OnlineSubsystemEOS/Private/OnlineSessionEOS.h b/Engine/Plugins/Online/OnlineSubsystemEOS/Source/OnlineSubsystemEOS/Private/OnlineSessionEOS.h
--- a/Engine/Plugins/Online/OnlineSubsystemEOS/Source/OnlineSubsystemEOS/Private/OnlineSessionEOS.h
+++ b/Engine/Plugins/Online/OnlineSubsystemEOS/Source/OnlineSubsystemEOS/Private/OnlineSessionEOS.h
@@ -16,8 +16,6 @@ class FOnlineSubsystemEOS;
 	#include "eos_lobby_types.h"
 
 
-typedef TEOSCallback<EOS_Sessions_OnUpdateSessionCallback, EOS_Sessions_UpdateSessionCallbackInfo> FUpdateSessionCallback;
-
 struct FSessionSearchEOS
 {
 	EOS_HSessionSearch SearchHandle;
@@ -36,8 +34,9 @@ struct FSessionSearchEOS
 /**
  * Interface for interacting with EOS sessions
  */
-class FOnlineSessionEOS :
-	public IOnlineSession
+class FOnlineSessionEOS
+	: public IOnlineSession
+	, public TSharedFromThis<FOnlineSessionEOS, ESPMode::ThreadSafe>
 {
 public:
 	FOnlineSessionEOS() = delete;
@@ -274,6 +273,7 @@ private:
 	void SetJoinInProgress(EOS_HSessionModification SessionModHandle, FNamedOnlineSession* Session);
 	void AddAttribute(EOS_HSessionModification SessionModHandle, const EOS_Sessions_AttributeData* Attribute);
 	void SetAttributes(EOS_HSessionModification SessionModHandle, FNamedOnlineSession* Session);
+	typedef TEOSCallback<EOS_Sessions_OnUpdateSessionCallback, EOS_Sessions_UpdateSessionCallbackInfo, FOnlineSessionEOS> FUpdateSessionCallback;
 	uint32 SharedSessionUpdate(EOS_HSessionModification SessionModHandle, FNamedOnlineSession* Session, FUpdateSessionCallback* Callback);
 
 	void TickLanTasks(float DeltaTime);
@@ -310,5 +310,6 @@ private:
 };
 
 typedef TSharedPtr<FOnlineSessionEOS, ESPMode::ThreadSafe> FOnlineSessionEOSPtr;
+typedef TWeakPtr<FOnlineSessionEOS, ESPMode::ThreadSafe> FOnlineSessionEOSWeakPtr;
 
 #endif
diff --git a/Engine/Plugins/Online/OnlineSubsystemEOS/Source/OnlineSubsystemEOS/Private/OnlineStatsEOS.cpp b/Engine/Plugins/Online/OnlineSubsystemEOS/Source/OnlineSubsystemEOS/Private/OnlineStatsEOS.cpp
--- a/Engine/Plugins/Online/OnlineSubsystemEOS/Source/OnlineSubsystemEOS/Private/OnlineStatsEOS.cpp
+++ b/Engine/Plugins/Online/OnlineSubsystemEOS/Source/OnlineSubsystemEOS/Private/OnlineStatsEOS.cpp
@@ -65,7 +65,7 @@ struct FQueryStatsOptions :
 	}
 };
 
-typedef TEOSCallback<EOS_Stats_OnQueryStatsCompleteCallback, EOS_Stats_OnQueryStatsCompleteCallbackInfo> FReadStatsCallback;
+typedef TEOSCallback<EOS_Stats_OnQueryStatsCompleteCallback, EOS_Stats_OnQueryStatsCompleteCallbackInfo, FOnlineStatsEOS> FReadStatsCallback;
 
 struct FStatsQueryContext
 {
@@ -156,7 +156,7 @@ void FOnlineStatsEOS::QueryStats(const FUniqueNetIdRef LocalUserId, const TArray
 		Options.LocalUserId = LocalEOSUserId;
 		Options.TargetUserId = TargetEOSUserId;
 
-		FReadStatsCallback* CallbackObj = new FReadStatsCallback();
+		FReadStatsCallback* CallbackObj = new FReadStatsCallback(FOnlineStatsEOSWeakPtr(AsShared()));
 		CallbackObj->CallbackLambda = [this, StatsQueryContext](const EOS_Stats_OnQueryStatsCompleteCallbackInfo* Data)
 		{
 			StatsQueryContext->NumPlayerReads--;
@@ -279,7 +279,7 @@ inline int32 GetVariantValue(const FOnlineStatValue& Data)
 	return Value;
 }
 
-typedef TEOSCallback<EOS_Stats_OnIngestStatCompleteCallback, EOS_Stats_IngestStatCompleteCallbackInfo> FWriteStatsCallback;
+typedef TEOSCallback<EOS_Stats_OnIngestStatCompleteCallback, EOS_Stats_IngestStatCompleteCallbackInfo, FOnlineStatsEOS> FWriteStatsCallback;
 
 void FOnlineStatsEOS::WriteStats(EOS_ProductUserId LocalUserId, EOS_ProductUserId UserId, const FOnlineStatsUserUpdatedStats& PlayerStats)
 {
@@ -309,7 +309,7 @@ void FOnlineStatsEOS::WriteStats(EOS_ProductUserId LocalUserId, EOS_ProductUserI
 	Options.Stats = EOSData.GetData();
 	Options.StatsCount = EOSData.Num();
 
-	FWriteStatsCallback* CallbackObj = new FWriteStatsCallback();
+	FWriteStatsCallback* CallbackObj = new FWriteStatsCallback(FOnlineStatsEOSWeakPtr(AsShared()));
 	CallbackObj->CallbackLambda = [this](const EOS_Stats_IngestStatCompleteCallbackInfo* Data)
 	{
 		bool bWasSuccessful = Data->ResultCode == EOS_EResult::EOS_Success;
diff --git a/Engine/Plugins/Online/OnlineSubsystemEOS/Source/OnlineSubsystemEOS/Private/OnlineStatsEOS.h b/Engine/Plugins/Online/OnlineSubsystemEOS/Source/OnlineSubsystemEOS/Private/OnlineStatsEOS.h
--- a/Engine/Plugins/Online/OnlineSubsystemEOS/Source/OnlineSubsystemEOS/Private/OnlineStatsEOS.h
+++ b/Engine/Plugins/Online/OnlineSubsystemEOS/Source/OnlineSubsystemEOS/Private/OnlineStatsEOS.h
@@ -14,8 +14,9 @@ class FOnlineSubsystemEOS;
 /**
  * Interface for interacting with EOS stats
  */
-class FOnlineStatsEOS :
-	public IOnlineStats
+class FOnlineStatsEOS
+	: public IOnlineStats
+	, public TSharedFromThis<FOnlineStatsEOS, ESPMode::ThreadSafe>
 {
 public:
 	FOnlineStatsEOS() = delete;
@@ -47,5 +48,6 @@ private:
 };
 
 typedef TSharedPtr<FOnlineStatsEOS, ESPMode::ThreadSafe> FOnlineStatsEOSPtr;
+typedef TWeakPtr<FOnlineStatsEOS, ESPMode::ThreadSafe> FOnlineStatsEOSWeakPtr;
 
 #endif
diff --git a/Engine/Plugins/Online/OnlineSubsystemEOS/Source/OnlineSubsystemEOS/Private/OnlineStoreEOS.cpp b/Engine/Plugins/Online/OnlineSubsystemEOS/Source/OnlineSubsystemEOS/Private/OnlineStoreEOS.cpp
--- a/Engine/Plugins/Online/OnlineSubsystemEOS/Source/OnlineSubsystemEOS/Private/OnlineStoreEOS.cpp
+++ b/Engine/Plugins/Online/OnlineSubsystemEOS/Source/OnlineSubsystemEOS/Private/OnlineStoreEOS.cpp
@@ -37,7 +37,7 @@ void FOnlineStoreEOS::QueryOffersById(const FUniqueNetId& UserId, const TArray<F
 	QueryOffers(UserId, Delegate);
 }
 
-typedef TEOSCallback<EOS_Ecom_OnQueryOffersCallback, EOS_Ecom_QueryOffersCallbackInfo> FQueryOffersCallback;
+typedef TEOSCallback<EOS_Ecom_OnQueryOffersCallback, EOS_Ecom_QueryOffersCallbackInfo, FOnlineStoreEOS> FQueryOffersCallback;
 
 void FOnlineStoreEOS::QueryOffers(const FUniqueNetId& UserId, const FOnQueryOnlineStoreOffersComplete& Delegate)
 {
@@ -60,7 +60,7 @@ void FOnlineStoreEOS::QueryOffers(const FUniqueNetId& UserId, const FOnQueryOnli
 	Options.ApiVersion = EOS_ECOM_QUERYOFFERS_API_LATEST;
 	Options.LocalUserId = AccountId;
 
-	FQueryOffersCallback* CallbackObj = new FQueryOffersCallback();
+	FQueryOffersCallback* CallbackObj = new FQueryOffersCallback(FOnlineStoreEOSWeakPtr(AsShared()));
 	CallbackObj->CallbackLambda = [this, OnComplete = FOnQueryOnlineStoreOffersComplete(Delegate)](const EOS_Ecom_QueryOffersCallbackInfo* Data)
 	{
 		EOS_EResult Result = Data->ResultCode;
@@ -138,7 +138,7 @@ TSharedPtr<FOnlineStoreOffer> FOnlineStoreEOS::GetOffer(const FUniqueOfferId& Of
 	return nullptr;
 }
 
-typedef TEOSCallback<EOS_Ecom_OnCheckoutCallback, EOS_Ecom_CheckoutCallbackInfo> FCheckoutCallback;
+typedef TEOSCallback<EOS_Ecom_OnCheckoutCallback, EOS_Ecom_CheckoutCallbackInfo, FOnlineStoreEOS> FCheckoutCallback;
 
 void FOnlineStoreEOS::Checkout(const FUniqueNetId& UserId, const FPurchaseCheckoutRequest& CheckoutRequest, const FOnPurchaseCheckoutComplete& Delegate)
 {
@@ -181,7 +181,7 @@ void FOnlineStoreEOS::Checkout(const FUniqueNetId& UserId, const FPurchaseChecko
 	Options.EntryCount = NumItems;
 	Options.Entries = (const EOS_Ecom_CheckoutEntry*)Entries.GetData();
 
-	FCheckoutCallback* CallbackObj = new FCheckoutCallback();
+	FCheckoutCallback* CallbackObj = new FCheckoutCallback(FOnlineStoreEOSWeakPtr(AsShared()));
 	CallbackObj->CallbackLambda = [this, OnComplete = FOnPurchaseCheckoutComplete(Delegate)](const EOS_Ecom_CheckoutCallbackInfo* Data)
 	{
 		EOS_EResult Result = Data->ResultCode;
@@ -232,7 +232,7 @@ void FOnlineStoreEOS::RedeemCode(const FUniqueNetId& UserId, const FRedeemCodeRe
 	Delegate.ExecuteIfBound(ONLINE_ERROR(EOnlineErrorResult::NotImplemented), BlankReceipt);
 }
 
-typedef TEOSCallback<EOS_Ecom_OnQueryEntitlementsCallback, EOS_Ecom_QueryEntitlementsCallbackInfo> FQueryReceiptsCallback;
+typedef TEOSCallback<EOS_Ecom_OnQueryEntitlementsCallback, EOS_Ecom_QueryEntitlementsCallbackInfo, FOnlineStoreEOS> FQueryReceiptsCallback;
 
 void FOnlineStoreEOS::QueryReceipts(const FUniqueNetId& UserId, bool bRestoreReceipts, const FOnQueryReceiptsComplete& Delegate)
 {
@@ -251,7 +251,7 @@ void FOnlineStoreEOS::QueryReceipts(const FUniqueNetId& UserId, bool bRestoreRec
 	Options.LocalUserId = AccountId;
 	Options.bIncludeRedeemed = bRestoreReceipts ? EOS_TRUE : EOS_FALSE;
 
-	FQueryReceiptsCallback* CallbackObj = new FQueryReceiptsCallback();
+	FQueryReceiptsCallback* CallbackObj = new FQueryReceiptsCallback(FOnlineStoreEOSWeakPtr(AsShared()));
 	CallbackObj->CallbackLambda = [this, OnComplete = FOnQueryReceiptsComplete(Delegate)](const EOS_Ecom_QueryEntitlementsCallbackInfo* Data)
 	{
 		EOS_EResult Result = Data->ResultCode;
@@ -307,7 +307,7 @@ void FOnlineStoreEOS::GetReceipts(const FUniqueNetId& UserId, TArray<FPurchaseRe
 	OutReceipts = CachedReceipts;
 }
 
-typedef TEOSCallback<EOS_Ecom_OnRedeemEntitlementsCallback, EOS_Ecom_RedeemEntitlementsCallbackInfo> FRedeemReceiptCallback;
+typedef TEOSCallback<EOS_Ecom_OnRedeemEntitlementsCallback, EOS_Ecom_RedeemEntitlementsCallbackInfo, FOnlineStoreEOS> FRedeemReceiptCallback;
 
 void FOnlineStoreEOS::FinalizeReceiptValidationInfo(const FUniqueNetId& UserId, FString& InReceiptValidationInfo, const FOnFinalizeReceiptValidationInfoComplete& Delegate)
 {
@@ -333,7 +333,7 @@ void FOnlineStoreEOS::FinalizeReceiptValidationInfo(const FUniqueNetId& UserId,
 	Options.EntitlementIdCount = 1;
 	Options.EntitlementIds = Ids;
 
-	FRedeemReceiptCallback* CallbackObj = new FRedeemReceiptCallback();
+	FRedeemReceiptCallback* CallbackObj = new FRedeemReceiptCallback(FOnlineStoreEOSWeakPtr(AsShared()));
 	CallbackObj->CallbackLambda = [this, Info = FString(InReceiptValidationInfo), OnComplete = FOnFinalizeReceiptValidationInfoComplete(Delegate)](const EOS_Ecom_RedeemEntitlementsCallbackInfo* Data)
 	{
 		EOS_EResult Result = Data->ResultCode;
diff --git a/Engine/Plugins/Online/OnlineSubsystemEOS/Source/OnlineSubsystemEOS/Private/OnlineStoreEOS.h b/Engine/Plugins/Online/OnlineSubsystemEOS/Source/OnlineSubsystemEOS/Private/OnlineStoreEOS.h
--- a/Engine/Plugins/Online/OnlineSubsystemEOS/Source/OnlineSubsystemEOS/Private/OnlineStoreEOS.h
+++ b/Engine/Plugins/Online/OnlineSubsystemEOS/Source/OnlineSubsystemEOS/Private/OnlineStoreEOS.h
@@ -69,5 +69,6 @@ private:
 };
 
 typedef TSharedPtr<FOnlineStoreEOS, ESPMode::ThreadSafe> FOnlineStoreEOSPtr;
+typedef TWeakPtr<FOnlineStoreEOS, ESPMode::ThreadSafe> FOnlineStoreEOSWeakPtr;
 
 #endif
diff --git a/Engine/Plugins/Online/OnlineSubsystemEOS/Source/OnlineSubsystemEOS/Private/OnlineSubsystemEOS.cpp b/Engine/Plugins/Online/OnlineSubsystemEOS/Source/OnlineSubsystemEOS/Private/OnlineSubsystemEOS.cpp
--- a/Engine/Plugins/Online/OnlineSubsystemEOS/Source/OnlineSubsystemEOS/Private/OnlineSubsystemEOS.cpp
+++ b/Engine/Plugins/Online/OnlineSubsystemEOS/Source/OnlineSubsystemEOS/Private/OnlineSubsystemEOS.cpp
@@ -249,8 +249,6 @@ bool FOnlineSubsystemEOS::PlatformCreate()
 
 bool FOnlineSubsystemEOS::Init()
 {
-	FCallbackBase::EnableAllCallbacks();
-
 	// Determine if we are the default and if we're the platform OSS
 	FString DefaultOSS;
 	GConfig->GetString(TEXT("OnlineSubsystem"), TEXT("DefaultPlatformService"), DefaultOSS, GEngineIni);
@@ -397,6 +395,7 @@ bool FOnlineSubsystemEOS::Init()
 	}
 
 	UserManager = MakeShareable(new FUserManagerEOS(this));
+	UserManager->Init();
 	SessionInterfacePtr = MakeShareable(new FOnlineSessionEOS(this));
 	// Set the bucket id to use for all sessions based upon the name and version to avoid upgrade issues
 	SessionInterfacePtr->Init(EOSSDKManager->GetProductName() + TEXT("_") + EOSSDKManager->GetProductVersion());
@@ -422,7 +421,6 @@ bool FOnlineSubsystemEOS::Shutdown()
 		EOS_Platform_Tick(*EOSPlatformHandle);
 	}
 
-	FCallbackBase::CancelAllCallbacks();
 	StopTicker();
 
 	if (SocketSubsystem)
diff --git a/Engine/Plugins/Online/OnlineSubsystemEOS/Source/OnlineSubsystemEOS/Private/OnlineSubsystemEOSTypes.h b/Engine/Plugins/Online/OnlineSubsystemEOS/Source/OnlineSubsystemEOS/Private/OnlineSubsystemEOSTypes.h
--- a/Engine/Plugins/Online/OnlineSubsystemEOS/Source/OnlineSubsystemEOS/Private/OnlineSubsystemEOSTypes.h
+++ b/Engine/Plugins/Online/OnlineSubsystemEOS/Source/OnlineSubsystemEOS/Private/OnlineSubsystemEOSTypes.h
@@ -438,16 +438,22 @@ static inline FString MakeNetIdStringFromIds(EOS_EpicAccountId AccountId, EOS_Pr
 }
 
 /** Class to handle all callbacks generically using a lambda to process callback results */
-template<typename CallbackFuncType, typename CallbackType>
+template<typename CallbackFuncType, typename CallbackType, typename OwningType>
 class TEOSCallback :
 	public FCallbackBase
 {
 public:
 	TFunction<void(const CallbackType*)> CallbackLambda;
 
-	TEOSCallback()
+	TEOSCallback(TWeakPtr<OwningType> InOwner)
+		: FCallbackBase()
+		, Owner(InOwner)
+	{
+	}
+	TEOSCallback(TWeakPtr<const OwningType> InOwner)
+		: FCallbackBase()
+		, Owner(InOwner)
 	{
-
 	}
 	virtual ~TEOSCallback() = default;
 
@@ -457,6 +463,10 @@ public:
 		return &CallbackImpl;
 	}
 
+protected:
+	/** The object that needs to be checked for lifetime before calling the callback */
+	TWeakPtr<const OwningType> Owner;
+
 private:
 	static void EOS_CALL CallbackImpl(const CallbackType* Data)
 	{
@@ -470,15 +480,11 @@ private:
 		TEOSCallback* CallbackThis = (TEOSCallback*)Data->ClientData;
 		check(CallbackThis);
 
-		if (FCallbackBase::ShouldCancelAllCallbacks())
+		if (CallbackThis->Owner.IsValid())
 		{
-			delete CallbackThis;
-			return;
+			check(CallbackThis->CallbackLambda);
+			CallbackThis->CallbackLambda(Data);
 		}
-
-		check(CallbackThis->CallbackLambda);
-		CallbackThis->CallbackLambda(Data);
-
 		delete CallbackThis;
 	}
 };
@@ -508,13 +514,16 @@ namespace OSSInternalCallback
  * Class to handle nested callbacks (callbacks that are tied to an external callback's lifetime,
  * e.g. file chunkers) generically using a lambda to process callback results
  */
-template<typename CallbackFuncType, typename CallbackType,
+template<typename CallbackFuncType, typename CallbackType, typename OwningType,
 	typename Nested1CallbackFuncType, typename Nested1CallbackType, typename Nested1ReturnType>
 class TEOSCallbackWithNested1 :
-	public TEOSCallback<CallbackFuncType, CallbackType>
+	public TEOSCallback<CallbackFuncType, CallbackType, OwningType>
 {
 public:
-	TEOSCallbackWithNested1() = default;
+	TEOSCallbackWithNested1(TWeakPtr<OwningType> InOwner)
+		: TEOSCallback<CallbackFuncType, CallbackType, OwningType>(InOwner)
+	{
+	}
 	virtual ~TEOSCallbackWithNested1() = default;
 
 
@@ -534,14 +543,14 @@ private:
 	static Nested1ReturnType EOS_CALL Nested1CallbackImpl(const Nested1CallbackType* Data)
 	{
 		check(IsInGameThread());
-		if (FCallbackBase::ShouldCancelAllCallbacks())
+		TEOSCallbackWithNested1* CallbackThis = (TEOSCallbackWithNested1*)Data->ClientData;
+		check(CallbackThis);
+
+		if (!CallbackThis->Owner.IsValid())
 		{
 			return Nested1ReturnType();
 		}
 
-		TEOSCallbackWithNested1* CallbackThis = (TEOSCallbackWithNested1*)Data->ClientData;
-		check(CallbackThis);
-
 		check(CallbackThis->CallbackLambda);
 		return CallbackThis->Nested1CallbackLambda(Data);
 	}
@@ -551,14 +560,17 @@ private:
  * Class to handle 2 nested callbacks (callbacks that are tied to an external callback's lifetime,
  * e.g. file chunkers) generically using a lambda to process callback results
  */
-template<typename CallbackFuncType, typename CallbackType,
+template<typename CallbackFuncType, typename CallbackType, typename OwningType,
 	typename Nested1CallbackFuncType, typename Nested1CallbackType, typename Nested1ReturnType,
 	typename Nested2CallbackFuncType, typename Nested2CallbackType>
 class TEOSCallbackWithNested2 :
-	public TEOSCallbackWithNested1<CallbackFuncType, CallbackType, Nested1CallbackFuncType, Nested1CallbackType, Nested1ReturnType>
+	public TEOSCallbackWithNested1<CallbackFuncType, CallbackType, OwningType, Nested1CallbackFuncType, Nested1CallbackType, Nested1ReturnType>
 {
 public:
-	TEOSCallbackWithNested2() = default;
+	TEOSCallbackWithNested2(TWeakPtr<OwningType> InOwner)
+		: TEOSCallbackWithNested1<CallbackFuncType, CallbackType, OwningType, Nested1CallbackFuncType, Nested1CallbackType, Nested1ReturnType>(InOwner)
+	{
+	}
 	virtual ~TEOSCallbackWithNested2() = default;
 
 
@@ -578,16 +590,14 @@ private:
 	static void EOS_CALL Nested2CallbackImpl(const Nested2CallbackType* Data)
 	{
 		check(IsInGameThread());
-		if (FCallbackBase::ShouldCancelAllCallbacks())
-		{
-			return;
-		}
-
 		TEOSCallbackWithNested2* CallbackThis = (TEOSCallbackWithNested2*)Data->ClientData;
 		check(CallbackThis);
 
-		check(CallbackThis->CallbackLambda);
-		CallbackThis->Nested2CallbackLambda(Data);
+		if (CallbackThis->Owner.IsValid())
+		{
+			check(CallbackThis->CallbackLambda);
+			CallbackThis->Nested2CallbackLambda(Data);
+		}
 	}
 };
 
@@ -595,13 +605,16 @@ private:
  * Class to handle nested callbacks (callbacks that are tied to an external callback's lifetime,
  * e.g. file chunkers) generically using a lambda to process callback results
  */
-template<typename CallbackFuncType, typename CallbackType,
+template<typename CallbackFuncType, typename CallbackType, typename OwningType,
 	typename Nested1CallbackFuncType, typename Nested1CallbackType, typename Nested1ReturnType>
-	class TEOSCallbackWithNested1Param3 :
-	public TEOSCallback<CallbackFuncType, CallbackType>
+class TEOSCallbackWithNested1Param3 :
+	public TEOSCallback<CallbackFuncType, CallbackType, OwningType>
 {
 public:
-	TEOSCallbackWithNested1Param3() = default;
+	TEOSCallbackWithNested1Param3(TWeakPtr<OwningType> InOwner)
+		: TEOSCallback<CallbackFuncType, CallbackType, OwningType>(InOwner)
+	{
+	}
 	virtual ~TEOSCallbackWithNested1Param3() = default;
 
 
@@ -621,14 +634,14 @@ private:
 	static Nested1ReturnType EOS_CALL Nested1CallbackImpl(const Nested1CallbackType* Data, void* OutDataBuffer, uint32_t* OutDataWritten)
 	{
 		check(IsInGameThread());
-		if (FCallbackBase::ShouldCancelAllCallbacks())
+		TEOSCallbackWithNested1Param3* CallbackThis = (TEOSCallbackWithNested1Param3*)Data->ClientData;
+		check(CallbackThis);
+
+		if (!CallbackThis->Owner.IsValid())
 		{
 			return Nested1ReturnType();
 		}
 
-		TEOSCallbackWithNested1Param3* CallbackThis = (TEOSCallbackWithNested1Param3*)Data->ClientData;
-		check(CallbackThis);
-
 		check(CallbackThis->CallbackLambda);
 		return CallbackThis->Nested1CallbackLambda(Data, OutDataBuffer, OutDataWritten);
 	}
@@ -638,14 +651,17 @@ private:
  * Class to handle 2 nested callbacks (callbacks that are tied to an external callback's lifetime,
  * e.g. file chunkers) generically using a lambda to process callback results
  */
-template<typename CallbackFuncType, typename CallbackType,
+template<typename CallbackFuncType, typename CallbackType, typename OwningType,
 	typename Nested1CallbackFuncType, typename Nested1CallbackType, typename Nested1ReturnType,
 	typename Nested2CallbackFuncType, typename Nested2CallbackType>
-	class TEOSCallbackWithNested2ForNested1Param3 :
-	public TEOSCallbackWithNested1Param3<CallbackFuncType, CallbackType, Nested1CallbackFuncType, Nested1CallbackType, Nested1ReturnType>
+class TEOSCallbackWithNested2ForNested1Param3 :
+	public TEOSCallbackWithNested1Param3<CallbackFuncType, CallbackType, OwningType, Nested1CallbackFuncType, Nested1CallbackType, Nested1ReturnType>
 {
 public:
-	TEOSCallbackWithNested2ForNested1Param3() = default;
+	TEOSCallbackWithNested2ForNested1Param3(TWeakPtr<OwningType> InOwner)
+		: TEOSCallbackWithNested1Param3<CallbackFuncType, CallbackType, OwningType, Nested1CallbackFuncType, Nested1CallbackType, Nested1ReturnType>(InOwner)
+	{
+	}
 	virtual ~TEOSCallbackWithNested2ForNested1Param3() = default;
 
 
@@ -665,16 +681,14 @@ private:
 	static void EOS_CALL Nested2CallbackImpl(const Nested2CallbackType* Data)
 	{
 		check(IsInGameThread());
-		if (FCallbackBase::ShouldCancelAllCallbacks())
-		{
-			return;
-		}
-
 		TEOSCallbackWithNested2ForNested1Param3* CallbackThis = (TEOSCallbackWithNested2ForNested1Param3*)Data->ClientData;
 		check(CallbackThis);
 
-		check(CallbackThis->CallbackLambda);
-		CallbackThis->Nested2CallbackLambda(Data);
+		if (CallbackThis->Owner.IsValid())
+		{
+			check(CallbackThis->CallbackLambda);
+			CallbackThis->Nested2CallbackLambda(Data);
+		}
 	}
 };
 
diff --git a/Engine/Plugins/Online/OnlineSubsystemEOS/Source/OnlineSubsystemEOS/Private/OnlineTitleFileEOS.cpp b/Engine/Plugins/Online/OnlineSubsystemEOS/Source/OnlineSubsystemEOS/Private/OnlineTitleFileEOS.cpp
--- a/Engine/Plugins/Online/OnlineSubsystemEOS/Source/OnlineSubsystemEOS/Private/OnlineTitleFileEOS.cpp
+++ b/Engine/Plugins/Online/OnlineSubsystemEOS/Source/OnlineSubsystemEOS/Private/OnlineTitleFileEOS.cpp
@@ -10,11 +10,11 @@
 #if WITH_EOS_SDK
 #include "eos_titlestorage.h"
 
-typedef TEOSCallback<EOS_TitleStorage_OnDeleteCacheCompleteCallback, EOS_TitleStorage_DeleteCacheCallbackInfo> FDeleteCacheCompleteCallback;
+typedef TEOSCallback<EOS_TitleStorage_OnDeleteCacheCompleteCallback, EOS_TitleStorage_DeleteCacheCallbackInfo, FOnlineTitleFileEOS> FDeleteCacheCompleteCallback;
 
-typedef TEOSCallback<EOS_TitleStorage_OnQueryFileListCompleteCallback, EOS_TitleStorage_QueryFileListCallbackInfo> FQueryFileListCallback;
+typedef TEOSCallback<EOS_TitleStorage_OnQueryFileListCompleteCallback, EOS_TitleStorage_QueryFileListCallbackInfo, FOnlineTitleFileEOS> FQueryFileListCallback;
 
-typedef TEOSCallbackWithNested2<EOS_TitleStorage_OnReadFileCompleteCallback, EOS_TitleStorage_ReadFileCallbackInfo,
+typedef TEOSCallbackWithNested2<EOS_TitleStorage_OnReadFileCompleteCallback, EOS_TitleStorage_ReadFileCallbackInfo, FOnlineTitleFileEOS,
 	EOS_TitleStorage_OnReadFileDataCallback, EOS_TitleStorage_ReadFileDataCallbackInfo, EOS_TitleStorage_EReadResult,
 	EOS_TitleStorage_OnFileTransferProgressCallback, EOS_TitleStorage_FileTransferProgressCallbackInfo
 > FReadTitleFileCompleteCallback;
@@ -97,7 +97,7 @@ void FOnlineTitleFileEOS::DeleteCachedFiles(bool bSkipEnumerated)
 	DeleteCacheOptions.ApiVersion = EOS_TITLESTORAGE_DELETECACHEOPTIONS_API_LATEST;
 	DeleteCacheOptions.LocalUserId = EOSSubsystem->UserManager->GetLocalProductUserId();	// Get a local user if one is available, but this is not required
 
-	FDeleteCacheCompleteCallback* CallbackObj = new FDeleteCacheCompleteCallback();
+	FDeleteCacheCompleteCallback* CallbackObj = new FDeleteCacheCompleteCallback(FOnlineTitleFileEOSWeakPtr(AsShared()));
 	CallbackObj->CallbackLambda = [this](const EOS_TitleStorage_DeleteCacheCallbackInfo* Data)
 	{
 		bool bWasSuccessful = Data->ResultCode == EOS_EResult::EOS_Success;
@@ -168,7 +168,7 @@ bool FOnlineTitleFileEOS::EnumerateFiles(const FPagedQuery& Page)
 	QueryFileListOptions.ListOfTags = const_cast<const char**>(AnsiTags.GetData());
 	QueryFileListOptions.ListOfTagsCount = AnsiTags.Num();
 
-	FQueryFileListCallback* CallbackObj = new FQueryFileListCallback();
+	FQueryFileListCallback* CallbackObj = new FQueryFileListCallback(FOnlineTitleFileEOSWeakPtr(AsShared()));
 	CallbackObj->CallbackLambda = [this](const EOS_TitleStorage_QueryFileListCallbackInfo* Data)
 	{
 		FString ErrorStr;
@@ -240,7 +240,7 @@ bool FOnlineTitleFileEOS::ReadFile(const FString& FileName)
 		return true;
 	}
 
-	FReadTitleFileCompleteCallback* CallbackObj = new FReadTitleFileCompleteCallback();
+	FReadTitleFileCompleteCallback* CallbackObj = new FReadTitleFileCompleteCallback(FOnlineTitleFileEOSWeakPtr(AsShared()));
 
 	CallbackObj->SetNested1CallbackLambda([this](const EOS_TitleStorage_ReadFileDataCallbackInfo* Data)
 	{
diff --git a/Engine/Plugins/Online/OnlineSubsystemEOS/Source/OnlineSubsystemEOS/Private/OnlineTitleFileEOS.h b/Engine/Plugins/Online/OnlineSubsystemEOS/Source/OnlineSubsystemEOS/Private/OnlineTitleFileEOS.h
--- a/Engine/Plugins/Online/OnlineSubsystemEOS/Source/OnlineSubsystemEOS/Private/OnlineTitleFileEOS.h
+++ b/Engine/Plugins/Online/OnlineSubsystemEOS/Source/OnlineSubsystemEOS/Private/OnlineTitleFileEOS.h
@@ -68,5 +68,6 @@ private:
 };
 
 typedef TSharedPtr<FOnlineTitleFileEOS, ESPMode::ThreadSafe> FOnlineTitleFileEOSPtr;
+typedef TWeakPtr<FOnlineTitleFileEOS, ESPMode::ThreadSafe> FOnlineTitleFileEOSWeakPtr;
 
 #endif
\ No newline at end of file
diff --git a/Engine/Plugins/Online/OnlineSubsystemEOS/Source/OnlineSubsystemEOS/Private/OnlineUserCloudEOS.cpp b/Engine/Plugins/Online/OnlineSubsystemEOS/Source/OnlineSubsystemEOS/Private/OnlineUserCloudEOS.cpp
--- a/Engine/Plugins/Online/OnlineSubsystemEOS/Source/OnlineSubsystemEOS/Private/OnlineUserCloudEOS.cpp
+++ b/Engine/Plugins/Online/OnlineSubsystemEOS/Source/OnlineSubsystemEOS/Private/OnlineUserCloudEOS.cpp
@@ -10,19 +10,19 @@
 #if WITH_EOS_SDK
 #include "eos_playerdatastorage.h"
 
-typedef TEOSCallback<EOS_PlayerDataStorage_OnQueryFileListCompleteCallback, EOS_PlayerDataStorage_QueryFileListCallbackInfo> FOnQueryFileListCallback;
+typedef TEOSCallback<EOS_PlayerDataStorage_OnQueryFileListCompleteCallback, EOS_PlayerDataStorage_QueryFileListCallbackInfo, FOnlineUserCloudEOS> FOnQueryFileListCallback;
 
-typedef TEOSCallbackWithNested2ForNested1Param3<EOS_PlayerDataStorage_OnWriteFileCompleteCallback, EOS_PlayerDataStorage_WriteFileCallbackInfo,
+typedef TEOSCallbackWithNested2ForNested1Param3<EOS_PlayerDataStorage_OnWriteFileCompleteCallback, EOS_PlayerDataStorage_WriteFileCallbackInfo, FOnlineUserCloudEOS,
 	EOS_PlayerDataStorage_OnWriteFileDataCallback, EOS_PlayerDataStorage_WriteFileDataCallbackInfo, EOS_PlayerDataStorage_EWriteResult,
 	EOS_PlayerDataStorage_OnFileTransferProgressCallback, EOS_PlayerDataStorage_FileTransferProgressCallbackInfo
 > FWriteUserFileCompleteCallback;
 
-typedef TEOSCallbackWithNested2<EOS_PlayerDataStorage_OnReadFileCompleteCallback, EOS_PlayerDataStorage_ReadFileCallbackInfo,
+typedef TEOSCallbackWithNested2<EOS_PlayerDataStorage_OnReadFileCompleteCallback, EOS_PlayerDataStorage_ReadFileCallbackInfo, FOnlineUserCloudEOS,
 	EOS_PlayerDataStorage_OnReadFileDataCallback, EOS_PlayerDataStorage_ReadFileDataCallbackInfo, EOS_PlayerDataStorage_EReadResult,
 	EOS_PlayerDataStorage_OnFileTransferProgressCallback, EOS_PlayerDataStorage_FileTransferProgressCallbackInfo
 > FReadUserFileCompleteCallback;
 
-typedef TEOSCallback<EOS_PlayerDataStorage_OnDeleteFileCompleteCallback, EOS_PlayerDataStorage_DeleteFileCallbackInfo> FOnDeleteFileCallback;
+typedef TEOSCallback<EOS_PlayerDataStorage_OnDeleteFileCompleteCallback, EOS_PlayerDataStorage_DeleteFileCallbackInfo, FOnlineUserCloudEOS> FOnDeleteFileCallback;
 
 void FEOSUserCloudFile::Unload()
 {
@@ -192,7 +192,7 @@ void FOnlineUserCloudEOS::EnumerateUserFiles(const FUniqueNetId& UserId)
 	Options.ApiVersion = EOS_PLAYERDATASTORAGE_QUERYFILELISTOPTIONS_API_LATEST;
 	Options.LocalUserId = LocalUserId;
 
-	FOnQueryFileListCallback* CallbackObj = new FOnQueryFileListCallback();
+	FOnQueryFileListCallback* CallbackObj = new FOnQueryFileListCallback(FOnlineUserCloudEOSWeakPtr(AsShared()));
 	CallbackObj->CallbackLambda = [this](const EOS_PlayerDataStorage_QueryFileListCallbackInfo* Data)
 	{
 		bool bWasSuccessful = Data->ResultCode == EOS_EResult::EOS_Success || Data->ResultCode == EOS_EResult::EOS_NotFound; // If the user doesn't have any files yet, we will get a NotFound error, but the query was valid
@@ -312,7 +312,7 @@ bool FOnlineUserCloudEOS::ReadUserFile(const FUniqueNetId& UserId, const FString
 		ReadChunkSize = 16 * 1024;
 	}
 
-	FReadUserFileCompleteCallback* CallbackObj = new FReadUserFileCompleteCallback();
+	FReadUserFileCompleteCallback* CallbackObj = new FReadUserFileCompleteCallback(FOnlineUserCloudEOSWeakPtr(AsShared()));
 
 	CallbackObj->SetNested1CallbackLambda([this, SharedUserId, FileName](const EOS_PlayerDataStorage_ReadFileDataCallbackInfo* Data)
 		{
@@ -493,7 +493,7 @@ bool FOnlineUserCloudEOS::WriteUserFile(const FUniqueNetId& UserId, const FStrin
 		ReadChunkSize = 16 * 1024;
 	}
 
-	FWriteUserFileCompleteCallback* CallbackObj = new FWriteUserFileCompleteCallback();
+	FWriteUserFileCompleteCallback* CallbackObj = new FWriteUserFileCompleteCallback(FOnlineUserCloudEOSWeakPtr(AsShared()));
 
 	CallbackObj->SetNested1CallbackLambda([this, SharedUserId, FileName](const EOS_PlayerDataStorage_WriteFileDataCallbackInfo* Data, void* OutDataBuffer, uint32_t* OutDataWritten)
 	{
@@ -751,7 +751,7 @@ bool FOnlineUserCloudEOS::DeleteUserFile(const FUniqueNetId& UserId, const FStri
 		Options.LocalUserId = EOSSubsystem->UserManager->GetProductUserId(UserId);
 		Options.Filename = FileNameUtf8.Get();
 
-		FOnDeleteFileCallback* CallbackObj = new FOnDeleteFileCallback();
+		FOnDeleteFileCallback* CallbackObj = new FOnDeleteFileCallback(FOnlineUserCloudEOSWeakPtr(AsShared()));
 		CallbackObj->CallbackLambda = [this, UserIdRef = UserId.AsShared(), FileName](const EOS_PlayerDataStorage_DeleteFileCallbackInfo* Data)
 		{
 			bool bWasSuccessful = Data->ResultCode == EOS_EResult::EOS_Success;
diff --git a/Engine/Plugins/Online/OnlineSubsystemEOS/Source/OnlineSubsystemEOS/Private/OnlineUserCloudEOS.h b/Engine/Plugins/Online/OnlineSubsystemEOS/Source/OnlineSubsystemEOS/Private/OnlineUserCloudEOS.h
--- a/Engine/Plugins/Online/OnlineSubsystemEOS/Source/OnlineSubsystemEOS/Private/OnlineUserCloudEOS.h
+++ b/Engine/Plugins/Online/OnlineSubsystemEOS/Source/OnlineSubsystemEOS/Private/OnlineUserCloudEOS.h
@@ -74,5 +74,6 @@ private:
 };
 
 typedef TSharedPtr<FOnlineUserCloudEOS, ESPMode::ThreadSafe> FOnlineUserCloudEOSPtr;
+typedef TWeakPtr<FOnlineUserCloudEOS, ESPMode::ThreadSafe> FOnlineUserCloudEOSWeakPtr;
 
 #endif
\ No newline at end of file
diff --git a/Engine/Plugins/Online/OnlineSubsystemEOS/Source/OnlineSubsystemEOS/Private/UserManagerEOS.cpp b/Engine/Plugins/Online/OnlineSubsystemEOS/Source/OnlineSubsystemEOS/Private/UserManagerEOS.cpp
--- a/Engine/Plugins/Online/OnlineSubsystemEOS/Source/OnlineSubsystemEOS/Private/UserManagerEOS.cpp
+++ b/Engine/Plugins/Online/OnlineSubsystemEOS/Source/OnlineSubsystemEOS/Private/UserManagerEOS.cpp
@@ -129,10 +129,11 @@ PRAGMA_ENABLE_DEPRECATION_WARNINGS
 IOnlinePresence::FOnPresenceTaskCompleteDelegate IgnoredPresenceDelegate;
 IOnlineUser::FOnQueryExternalIdMappingsComplete IgnoredMappingDelegate;
 
-typedef TEOSGlobalCallback<EOS_UI_OnDisplaySettingsUpdatedCallback, EOS_UI_OnDisplaySettingsUpdatedCallbackInfo> FOnDisplaySettingsUpdatedCallback;
+typedef TEOSGlobalCallback<EOS_UI_OnDisplaySettingsUpdatedCallback, EOS_UI_OnDisplaySettingsUpdatedCallbackInfo, FUserManagerEOS> FOnDisplaySettingsUpdatedCallback;
 
 FUserManagerEOS::FUserManagerEOS(FOnlineSubsystemEOS* InSubsystem)
-	: EOSSubsystem(InSubsystem)
+	: TSharedFromThis<FUserManagerEOS, ESPMode::ThreadSafe>()
+	, EOSSubsystem(InSubsystem)
 	, DefaultLocalUser(-1)
 	, LoginNotificationId(0)
 	, LoginNotificationCallback(nullptr)
@@ -140,6 +141,10 @@ FUserManagerEOS::FUserManagerEOS(FOnlineSubsystemEOS* InSubsystem)
 	, FriendsNotificationCallback(nullptr)
 	, PresenceNotificationId(0)
 	, PresenceNotificationCallback(nullptr)
+{
+}
+
+void FUserManagerEOS::Init()
 {
 	// This delegate would cause a crash when running a dedicated server
 	if (!IsRunningDedicatedServer())
@@ -148,18 +153,18 @@ FUserManagerEOS::FUserManagerEOS(FOnlineSubsystemEOS* InSubsystem)
 		EOS_UI_AddNotifyDisplaySettingsUpdatedOptions Options = {};
 		Options.ApiVersion = EOS_UI_ADDNOTIFYDISPLAYSETTINGSUPDATED_API_LATEST;
 
-		FOnDisplaySettingsUpdatedCallback* CallbackObj = new FOnDisplaySettingsUpdatedCallback();
-		DisplaySettingsUpdatedCallback = CallbackObj;
-		CallbackObj->CallbackLambda = [this](const EOS_UI_OnDisplaySettingsUpdatedCallbackInfo* Data)
-		{
-			TriggerOnExternalUIChangeDelegates((bool)Data->bIsVisible);
-		};
+	FOnDisplaySettingsUpdatedCallback* CallbackObj = new FOnDisplaySettingsUpdatedCallback(FUserManagerEOSWeakPtr(AsShared()));
+	DisplaySettingsUpdatedCallback = CallbackObj;
+	CallbackObj->CallbackLambda = [this](const EOS_UI_OnDisplaySettingsUpdatedCallbackInfo* Data)
+	{
+		TriggerOnExternalUIChangeDelegates((bool)Data->bIsVisible);
+	};
 
 		DisplaySettingsUpdatedId = EOS_UI_AddNotifyDisplaySettingsUpdated(EOSSubsystem->UIHandle, &Options, CallbackObj, CallbackObj->GetCallbackPtr());
 	}
 }
 
-FUserManagerEOS::~FUserManagerEOS()
+void FUserManagerEOS::Shutdown()
 {
 	// This delegate would cause a crash when running a dedicated server
 	if (DisplaySettingsUpdatedId != EOS_INVALID_NOTIFICATIONID)
@@ -171,6 +176,11 @@ FUserManagerEOS::~FUserManagerEOS()
 	}
 }
 
+FUserManagerEOS::~FUserManagerEOS()
+{
+	Shutdown();
+}
+
 void FUserManagerEOS::LoginStatusChanged(const EOS_Auth_LoginStatusChangedCallbackInfo* Data)
 {
 	if (Data->CurrentStatus == EOS_ELoginStatus::EOS_LS_NotLoggedIn)
@@ -260,9 +270,9 @@ void FUserManagerEOS::GetPlatformAuthToken(int32 LocalUserNum, const FOnGetLinke
 	PlatformIdentity->GetLinkedAccountAuthToken(LocalUserNum, Delegate);
 }
 
-typedef TEOSCallback<EOS_Auth_OnLoginCallback, EOS_Auth_LoginCallbackInfo> FLoginCallback;
-typedef TEOSCallback<EOS_Connect_OnLoginCallback, EOS_Connect_LoginCallbackInfo> FConnectLoginCallback;
-typedef TEOSCallback<EOS_Auth_OnDeletePersistentAuthCallback, EOS_Auth_DeletePersistentAuthCallbackInfo> FDeletePersistentAuthCallback;
+typedef TEOSCallback<EOS_Auth_OnLoginCallback, EOS_Auth_LoginCallbackInfo, FUserManagerEOS> FLoginCallback;
+typedef TEOSCallback<EOS_Connect_OnLoginCallback, EOS_Connect_LoginCallbackInfo, FUserManagerEOS> FConnectLoginCallback;
+typedef TEOSCallback<EOS_Auth_OnDeletePersistentAuthCallback, EOS_Auth_DeletePersistentAuthCallbackInfo, FUserManagerEOS> FDeletePersistentAuthCallback;
 
 // Chose arbitrarily since the SDK doesn't define it
 #define EOS_MAX_TOKEN_SIZE 4096
@@ -437,7 +447,7 @@ bool FUserManagerEOS::Login(int32 LocalUserNum, const FOnlineAccountCredentials&
 		return false;
 	}
 
-	FLoginCallback* CallbackObj = new FLoginCallback();
+	FLoginCallback* CallbackObj = new FLoginCallback(FUserManagerEOSWeakPtr(AsShared()));
 	CallbackObj->CallbackLambda = [this, LocalUserNum, bIsPersistentLogin](const EOS_Auth_LoginCallbackInfo* Data)
 	{
 		if (Data->ResultCode == EOS_EResult::EOS_Success)
@@ -462,7 +472,7 @@ bool FUserManagerEOS::Login(int32 LocalUserNum, const FOnlineAccountCredentials&
 			// Check for invalid persistent login credentials.
 			if (bIsPersistentLogin && bShouldRemoveCachedToken)
 			{
-				FDeletePersistentAuthCallback* DeleteAuthCallbackObj = new FDeletePersistentAuthCallback();
+				FDeletePersistentAuthCallback* DeleteAuthCallbackObj = new FDeletePersistentAuthCallback(FUserManagerEOSWeakPtr(AsShared()));
 				DeleteAuthCallbackObj->CallbackLambda = [this, LocalUserNum, TriggerLoginFailure](const EOS_Auth_DeletePersistentAuthCallbackInfo* Data)
 				{
 					// Deleting the auth token is best effort.
@@ -505,7 +515,7 @@ void FUserManagerEOS::LoginViaExternalAuth(int32 LocalUserNum)
 			FAuthCredentials Credentials(ToEOS_EExternalCredentialType(GetPlatformOSS()->GetSubsystemName(), *LocalUserNumToLastLoginCredentials[LocalUserNum]), AuthToken);
 			LoginOptions.Credentials = &Credentials;
 
-			FLoginCallback* CallbackObj = new FLoginCallback();
+			FLoginCallback* CallbackObj = new FLoginCallback(FUserManagerEOSWeakPtr(AsShared()));
 			CallbackObj->CallbackLambda = [this, LocalUserNum](const EOS_Auth_LoginCallbackInfo* Data)
 			{
 				if (Data->ResultCode == EOS_EResult::EOS_Success)
@@ -540,12 +550,12 @@ struct FLinkAccountOptions :
 	}
 };
 
-typedef TEOSCallback<EOS_Auth_OnLinkAccountCallback, EOS_Auth_LinkAccountCallbackInfo> FLinkAccountCallback;
+typedef TEOSCallback<EOS_Auth_OnLinkAccountCallback, EOS_Auth_LinkAccountCallbackInfo, FUserManagerEOS> FLinkAccountCallback;
 
 void FUserManagerEOS::LinkEAS(int32 LocalUserNum, EOS_ContinuanceToken Token)
 {
 	FLinkAccountOptions Options(Token);
-	FLinkAccountCallback* CallbackObj = new FLinkAccountCallback();
+	FLinkAccountCallback* CallbackObj = new FLinkAccountCallback(FUserManagerEOSWeakPtr(AsShared()));
 	CallbackObj->CallbackLambda = [this, LocalUserNum](const EOS_Auth_LinkAccountCallbackInfo* Data)
 	{
 		if (Data->ResultCode == EOS_EResult::EOS_Success)
@@ -624,7 +634,7 @@ bool FUserManagerEOS::ConnectLoginNoEAS(int32 LocalUserNum)
 			Options.ApiVersion = EOS_CONNECT_LOGIN_API_LATEST;
 			Options.Credentials = &Credentials;
 
-			FConnectLoginCallback* CallbackObj = new FConnectLoginCallback();
+			FConnectLoginCallback* CallbackObj = new FConnectLoginCallback(FUserManagerEOSWeakPtr(AsShared()));
 			CallbackObj->CallbackLambda = [this, LocalUserNum](const EOS_Connect_LoginCallbackInfo* Data)
 			{
 				if (Data->ResultCode == EOS_EResult::EOS_Success)
@@ -666,7 +676,7 @@ bool FUserManagerEOS::ConnectLoginEAS(int32 LocalUserNum, EOS_EpicAccountId Acco
 		Options.ApiVersion = EOS_CONNECT_LOGIN_API_LATEST;
 		Options.Credentials = &Credentials;
 
-		FConnectLoginCallback* CallbackObj = new FConnectLoginCallback();
+		FConnectLoginCallback* CallbackObj = new FConnectLoginCallback(FUserManagerEOSWeakPtr(AsShared()));
 		CallbackObj->CallbackLambda = [LocalUserNum, AccountId, this](const EOS_Connect_LoginCallbackInfo* Data)
 		{
 			if (Data->ResultCode == EOS_EResult::EOS_Success)
@@ -725,7 +735,7 @@ void FUserManagerEOS::RefreshConnectLogin(int32 LocalUserNum)
 			Options.ApiVersion = EOS_CONNECT_LOGIN_API_LATEST;
 			Options.Credentials = &Credentials;
 
-			FConnectLoginCallback* CallbackObj = new FConnectLoginCallback();
+			FConnectLoginCallback* CallbackObj = new FConnectLoginCallback(FUserManagerEOSWeakPtr(AsShared()));
 			CallbackObj->CallbackLambda = [LocalUserNum, AccountId, this](const EOS_Connect_LoginCallbackInfo* Data)
 			{
 				if (Data->ResultCode != EOS_EResult::EOS_Success)
@@ -766,7 +776,7 @@ void FUserManagerEOS::RefreshConnectLogin(int32 LocalUserNum)
 				Options.ApiVersion = EOS_CONNECT_LOGIN_API_LATEST;
 				Options.Credentials = &Credentials;
 
-				FConnectLoginCallback* CallbackObj = new FConnectLoginCallback();
+				FConnectLoginCallback* CallbackObj = new FConnectLoginCallback(FUserManagerEOSWeakPtr(AsShared()));
 				CallbackObj->CallbackLambda = [this, LocalUserNum](const EOS_Connect_LoginCallbackInfo* Data)
 				{
 					if (Data->ResultCode != EOS_EResult::EOS_Success)
@@ -780,7 +790,7 @@ void FUserManagerEOS::RefreshConnectLogin(int32 LocalUserNum)
 	}
 }
 
-typedef TEOSCallback<EOS_Connect_OnCreateUserCallback, EOS_Connect_CreateUserCallbackInfo> FCreateUserCallback;
+typedef TEOSCallback<EOS_Connect_OnCreateUserCallback, EOS_Connect_CreateUserCallbackInfo, FUserManagerEOS> FCreateUserCallback;
 
 void FUserManagerEOS::CreateConnectedLogin(int32 LocalUserNum, EOS_EpicAccountId AccountId, EOS_ContinuanceToken Token)
 {
@@ -788,7 +798,7 @@ void FUserManagerEOS::CreateConnectedLogin(int32 LocalUserNum, EOS_EpicAccountId
 	Options.ApiVersion = EOS_CONNECT_CREATEUSER_API_LATEST;
 	Options.ContinuanceToken = Token;
 
-	FCreateUserCallback* CallbackObj = new FCreateUserCallback();
+	FCreateUserCallback* CallbackObj = new FCreateUserCallback(FUserManagerEOSWeakPtr(AsShared()));
 	CallbackObj->CallbackLambda = [LocalUserNum, AccountId, this](const EOS_Connect_CreateUserCallbackInfo* Data)
 	{
 		if (Data->ResultCode == EOS_EResult::EOS_Success)
@@ -805,17 +815,17 @@ void FUserManagerEOS::CreateConnectedLogin(int32 LocalUserNum, EOS_EpicAccountId
 	EOS_Connect_CreateUser(EOSSubsystem->ConnectHandle, &Options, CallbackObj, CallbackObj->GetCallbackPtr());
 }
 
-typedef TEOSGlobalCallback<EOS_Connect_OnAuthExpirationCallback, EOS_Connect_AuthExpirationCallbackInfo> FRefreshAuthCallback;
-typedef TEOSGlobalCallback<EOS_Presence_OnPresenceChangedCallback, EOS_Presence_PresenceChangedCallbackInfo> FPresenceChangedCallback;
-typedef TEOSGlobalCallback<EOS_Friends_OnFriendsUpdateCallback, EOS_Friends_OnFriendsUpdateInfo> FFriendsStatusUpdateCallback;
-typedef TEOSGlobalCallback<EOS_Auth_OnLoginStatusChangedCallback, EOS_Auth_LoginStatusChangedCallbackInfo> FLoginStatusChangedCallback;
+typedef TEOSGlobalCallback<EOS_Connect_OnAuthExpirationCallback, EOS_Connect_AuthExpirationCallbackInfo, FUserManagerEOS> FRefreshAuthCallback;
+typedef TEOSGlobalCallback<EOS_Presence_OnPresenceChangedCallback, EOS_Presence_PresenceChangedCallbackInfo, FUserManagerEOS> FPresenceChangedCallback;
+typedef TEOSGlobalCallback<EOS_Friends_OnFriendsUpdateCallback, EOS_Friends_OnFriendsUpdateInfo, FUserManagerEOS> FFriendsStatusUpdateCallback;
+typedef TEOSGlobalCallback<EOS_Auth_OnLoginStatusChangedCallback, EOS_Auth_LoginStatusChangedCallbackInfo, FUserManagerEOS> FLoginStatusChangedCallback;
 
 void FUserManagerEOS::FullLoginCallback(int32 LocalUserNum, EOS_EpicAccountId AccountId, EOS_ProductUserId UserId)
 {
 	// Add our login status changed callback if not already set
 	if (LoginNotificationId == 0)
 	{
-		FLoginStatusChangedCallback* CallbackObj = new FLoginStatusChangedCallback();
+		FLoginStatusChangedCallback* CallbackObj = new FLoginStatusChangedCallback(FUserManagerEOSWeakPtr(AsShared()));
 		LoginNotificationCallback = CallbackObj;
 		CallbackObj->CallbackLambda = [this](const EOS_Auth_LoginStatusChangedCallbackInfo* Data)
 		{
@@ -829,7 +839,7 @@ void FUserManagerEOS::FullLoginCallback(int32 LocalUserNum, EOS_EpicAccountId Ac
 	// Register for friends updates if not set yet
 	if (FriendsNotificationId == 0)
 	{
-		FFriendsStatusUpdateCallback* CallbackObj = new FFriendsStatusUpdateCallback();
+		FFriendsStatusUpdateCallback* CallbackObj = new FFriendsStatusUpdateCallback(FUserManagerEOSWeakPtr(AsShared()));
 		FriendsNotificationCallback = CallbackObj;
 		CallbackObj->CallbackLambda = [LocalUserNum, this](const EOS_Friends_OnFriendsUpdateInfo* Data)
 		{
@@ -843,7 +853,7 @@ void FUserManagerEOS::FullLoginCallback(int32 LocalUserNum, EOS_EpicAccountId Ac
 	// Register for presence updates if not set yet
 	if (PresenceNotificationId == 0)
 	{
-		FPresenceChangedCallback* CallbackObj = new FPresenceChangedCallback();
+		FPresenceChangedCallback* CallbackObj = new FPresenceChangedCallback(FUserManagerEOSWeakPtr(AsShared()));
 		PresenceNotificationCallback = CallbackObj;
 		CallbackObj->CallbackLambda = [LocalUserNum, this](const EOS_Presence_PresenceChangedCallbackInfo* Data)
 		{
@@ -865,7 +875,7 @@ void FUserManagerEOS::FullLoginCallback(int32 LocalUserNum, EOS_EpicAccountId Ac
 		FNotificationIdCallbackPair* NotificationPair = new FNotificationIdCallbackPair();
 		LocalUserNumToConnectLoginNotifcationMap.Emplace(LocalUserNum, NotificationPair);
 
-		FRefreshAuthCallback* CallbackObj = new FRefreshAuthCallback();
+		FRefreshAuthCallback* CallbackObj = new FRefreshAuthCallback(FUserManagerEOSWeakPtr(AsShared()));
 		NotificationPair->Callback = CallbackObj;
 		CallbackObj->CallbackLambda = [LocalUserNum, this](const EOS_Connect_AuthExpirationCallbackInfo* Data)
 		{
@@ -885,7 +895,7 @@ void FUserManagerEOS::FullLoginCallback(int32 LocalUserNum, EOS_EpicAccountId Ac
 	TriggerOnLoginStatusChangedDelegates(LocalUserNum, ELoginStatus::NotLoggedIn, ELoginStatus::LoggedIn, *UserNetId);
 }
 
-typedef TEOSCallback<EOS_Auth_OnLogoutCallback, EOS_Auth_LogoutCallbackInfo> FLogoutCallback;
+typedef TEOSCallback<EOS_Auth_OnLogoutCallback, EOS_Auth_LogoutCallbackInfo, FUserManagerEOS> FLogoutCallback;
 
 bool FUserManagerEOS::Logout(int32 LocalUserNum)
 {
@@ -898,10 +908,10 @@ bool FUserManagerEOS::Logout(int32 LocalUserNum)
 		return false;
 	}
 
-	FLogoutCallback* CallbackObj = new FLogoutCallback();
+	FLogoutCallback* CallbackObj = new FLogoutCallback(FUserManagerEOSWeakPtr(AsShared()));
 	CallbackObj->CallbackLambda = [LocalUserNum, this](const EOS_Auth_LogoutCallbackInfo* Data)
 	{
-		FDeletePersistentAuthCallback* DeleteAuthCallbackObj = new FDeletePersistentAuthCallback();
+		FDeletePersistentAuthCallback* DeleteAuthCallbackObj = new FDeletePersistentAuthCallback(FUserManagerEOSWeakPtr(AsShared()));
 		DeleteAuthCallbackObj->CallbackLambda = [this, LocalUserNum, LogoutResultCode = Data->ResultCode](const EOS_Auth_DeletePersistentAuthCallbackInfo* Data)
 		{
 			if (LogoutResultCode == EOS_EResult::EOS_Success)
@@ -1200,7 +1210,7 @@ EOS_ProductUserId FUserManagerEOS::GetProductUserId(const FUniqueNetId& NetId) c
 	return nullptr;
 }
 
-typedef TEOSCallback<EOS_Connect_OnQueryProductUserIdMappingsCallback, EOS_Connect_QueryProductUserIdMappingsCallbackInfo> FConnectQueryProductUserIdMappingsCallback;
+typedef TEOSCallback<EOS_Connect_OnQueryProductUserIdMappingsCallback, EOS_Connect_QueryProductUserIdMappingsCallbackInfo, FUserManagerEOS> FConnectQueryProductUserIdMappingsCallback;
 
 /**
  * Uses the Connect API to retrieve the EOS_EpicAccountId for a given EOS_ProductUserId
@@ -1256,7 +1266,7 @@ void FUserManagerEOS::GetEpicAccountIdAsync(const EOS_ProductUserId& ProductUser
 		QueryProductUserIdMappingsOptions.ProductUserIds = ProductUserIdList.GetData();
 		QueryProductUserIdMappingsOptions.ProductUserIdCount = 1;
 
-		FConnectQueryProductUserIdMappingsCallback* CallbackObj = new FConnectQueryProductUserIdMappingsCallback();
+		FConnectQueryProductUserIdMappingsCallback* CallbackObj = new FConnectQueryProductUserIdMappingsCallback(FUserManagerEOSConstWeakPtr(AsShared()));
 		CallbackObj->CallbackLambda = [this, ProductUserId, Callback](const EOS_Connect_QueryProductUserIdMappingsCallbackInfo* Data)
 		{
 			EOS_EpicAccountId AccountId = nullptr;
@@ -1506,7 +1516,7 @@ bool FUserManagerEOS::ShowAccountCreationUI(const int ControllerIndex, const FOn
 	return true;
 }
 
-typedef TEOSCallback<EOS_UI_OnShowFriendsCallback, EOS_UI_ShowFriendsCallbackInfo> FOnShowFriendsCallback;
+typedef TEOSCallback<EOS_UI_OnShowFriendsCallback, EOS_UI_ShowFriendsCallbackInfo, FUserManagerEOS> FOnShowFriendsCallback;
 
 bool FUserManagerEOS::ShowFriendsUI(int32 LocalUserNum)
 {
@@ -1514,7 +1524,7 @@ bool FUserManagerEOS::ShowFriendsUI(int32 LocalUserNum)
 	Options.ApiVersion = EOS_UI_SHOWFRIENDS_API_LATEST;
 	Options.LocalUserId = GetLocalEpicAccountId(LocalUserNum);
 
-	FOnShowFriendsCallback* CallbackObj = new FOnShowFriendsCallback();
+	FOnShowFriendsCallback* CallbackObj = new FOnShowFriendsCallback(FUserManagerEOSWeakPtr(AsShared()));
 	CallbackObj->CallbackLambda = [this](const EOS_UI_ShowFriendsCallbackInfo* Data)
 	{
 		if (Data->ResultCode == EOS_EResult::EOS_Success)
@@ -1630,7 +1640,7 @@ bool FUserManagerEOS::ShowSendMessageUI(int32 LocalUserNum, const FShowSendMessa
 
 // ~IOnlineExternalUI Interface
 
-typedef TEOSCallback<EOS_Friends_OnQueryFriendsCallback, EOS_Friends_QueryFriendsCallbackInfo> FReadFriendsCallback;
+typedef TEOSCallback<EOS_Friends_OnQueryFriendsCallback, EOS_Friends_QueryFriendsCallbackInfo, FUserManagerEOS> FReadFriendsCallback;
 
 void FUserManagerEOS::FriendStatusChanged(const EOS_Friends_OnFriendsUpdateInfo* Data)
 {
@@ -1822,7 +1832,7 @@ bool FUserManagerEOS::ReadFriendsList(int32 LocalUserNum, const FString& ListNam
 	Options.ApiVersion = EOS_FRIENDS_QUERYFRIENDS_API_LATEST;
 	Options.LocalUserId = UserNumToAccountIdMap[LocalUserNum];
 
-	FReadFriendsCallback* CallbackObj = new FReadFriendsCallback();
+	FReadFriendsCallback* CallbackObj = new FReadFriendsCallback(FUserManagerEOSWeakPtr(AsShared()));
 	CallbackObj->CallbackLambda = [this, LocalUserNum, ListName, Delegate](const EOS_Friends_QueryFriendsCallbackInfo* Data)
 	{
 		EOS_EResult Result = Data->ResultCode;
@@ -1935,7 +1945,7 @@ bool FUserManagerEOS::DeleteFriendsList(int32 LocalUserNum, const FString& ListN
 	return true;
 }
 
-typedef TEOSCallback<EOS_Friends_OnSendInviteCallback, EOS_Friends_SendInviteCallbackInfo> FSendInviteCallback;
+typedef TEOSCallback<EOS_Friends_OnSendInviteCallback, EOS_Friends_SendInviteCallbackInfo, FUserManagerEOS> FSendInviteCallback;
 
 bool FUserManagerEOS::SendInvite(int32 LocalUserNum, const FUniqueNetId& FriendId, const FString& ListName, const FOnSendInviteComplete& Delegate)
 {
@@ -1954,7 +1964,7 @@ bool FUserManagerEOS::SendInvite(int32 LocalUserNum, const FUniqueNetId& FriendI
 		return false;
 	}
 
-	FSendInviteCallback* CallbackObj = new FSendInviteCallback();
+	FSendInviteCallback* CallbackObj = new FSendInviteCallback(FUserManagerEOSWeakPtr(AsShared()));
 	CallbackObj->CallbackLambda = [LocalUserNum, ListName, this, Delegate](const EOS_Friends_SendInviteCallbackInfo* Data)
 	{
 		const FString& NetId = AccountIdToStringMap[Data->TargetUserId];
@@ -1977,7 +1987,7 @@ bool FUserManagerEOS::SendInvite(int32 LocalUserNum, const FUniqueNetId& FriendI
 	return true;
 }
 
-typedef TEOSCallback<EOS_Friends_OnAcceptInviteCallback, EOS_Friends_AcceptInviteCallbackInfo> FAcceptInviteCallback;
+typedef TEOSCallback<EOS_Friends_OnAcceptInviteCallback, EOS_Friends_AcceptInviteCallbackInfo, FUserManagerEOS> FAcceptInviteCallback;
 
 bool FUserManagerEOS::AcceptInvite(int32 LocalUserNum, const FUniqueNetId& FriendId, const FString& ListName, const FOnAcceptInviteComplete& Delegate)
 {
@@ -1996,7 +2006,7 @@ bool FUserManagerEOS::AcceptInvite(int32 LocalUserNum, const FUniqueNetId& Frien
 		return false;
 	}
 
-	FAcceptInviteCallback* CallbackObj = new FAcceptInviteCallback();
+	FAcceptInviteCallback* CallbackObj = new FAcceptInviteCallback(FUserManagerEOSWeakPtr(AsShared()));
 	CallbackObj->CallbackLambda = [LocalUserNum, ListName, this, Delegate](const EOS_Friends_AcceptInviteCallbackInfo* Data)
 	{
 		const FString& NetId = AccountIdToStringMap[Data->TargetUserId];
@@ -2330,7 +2340,7 @@ struct FRichTextOptions :
 	char RichTextAnsi[EOS_PRESENCE_RICH_TEXT_MAX_VALUE_LENGTH];
 };
 
-typedef TEOSCallback<EOS_Presence_SetPresenceCompleteCallback, EOS_Presence_SetPresenceCallbackInfo> FSetPresenceCallback;
+typedef TEOSCallback<EOS_Presence_SetPresenceCompleteCallback, EOS_Presence_SetPresenceCallbackInfo, FUserManagerEOS> FSetPresenceCallback;
 
 void FUserManagerEOS::SetPresence(const FUniqueNetId& UserId, const FOnlineUserPresenceStatus& Status, const FOnPresenceTaskCompleteDelegate& Delegate)
 {
@@ -2393,7 +2403,7 @@ void FUserManagerEOS::SetPresence(const FUniqueNetId& UserId, const FOnlineUserP
 		UE_LOG_ONLINE(Error, TEXT("EOS_PresenceModification_SetData() failed with result code (%s)"), ANSI_TO_TCHAR(EOS_EResult_ToString(SetDataResult)));
 	}
 
-	FSetPresenceCallback* CallbackObj = new FSetPresenceCallback();
+	FSetPresenceCallback* CallbackObj = new FSetPresenceCallback(FUserManagerEOSWeakPtr(AsShared()));
 	CallbackObj->CallbackLambda = [this, Delegate](const EOS_Presence_SetPresenceCallbackInfo* Data)
 	{
 		if (Data->ResultCode == EOS_EResult::EOS_Success && AccountIdToStringMap.Contains(Data->LocalUserId))
@@ -2415,7 +2425,7 @@ void FUserManagerEOS::SetPresence(const FUniqueNetId& UserId, const FOnlineUserP
 	EOS_PresenceModification_Release(ChangeHandle);
 }
 
-typedef TEOSCallback<EOS_Presence_OnQueryPresenceCompleteCallback, EOS_Presence_QueryPresenceCallbackInfo> FQueryPresenceCallback;
+typedef TEOSCallback<EOS_Presence_OnQueryPresenceCompleteCallback, EOS_Presence_QueryPresenceCallbackInfo, FUserManagerEOS> FQueryPresenceCallback;
 
 void FUserManagerEOS::QueryPresence(const FUniqueNetId& UserId, const FOnPresenceTaskCompleteDelegate& Delegate)
 {
@@ -2442,7 +2452,7 @@ void FUserManagerEOS::QueryPresence(const FUniqueNetId& UserId, const FOnPresenc
 	EOS_Bool bHasPresence = EOS_Presence_HasPresence(EOSSubsystem->PresenceHandle, &HasOptions);
 	if (bHasPresence == EOS_FALSE)
 	{
-		FQueryPresenceCallback* CallbackObj = new FQueryPresenceCallback();
+		FQueryPresenceCallback* CallbackObj = new FQueryPresenceCallback(FUserManagerEOSWeakPtr(AsShared()));
 		CallbackObj->CallbackLambda = [this, Delegate](const EOS_Presence_QueryPresenceCallbackInfo* Data)
 		{
 			if (Data->ResultCode == EOS_EResult::EOS_Success && EpicAccountIdToOnlineUserMap.Contains(Data->TargetUserId))
@@ -2592,11 +2602,11 @@ bool FUserManagerEOS::QueryUserInfo(int32 LocalUserNum, const TArray<FUniqueNetI
 	return true;
 }
 
-typedef TEOSCallback<EOS_UserInfo_OnQueryUserInfoCallback, EOS_UserInfo_QueryUserInfoCallbackInfo> FReadUserInfoCallback;
+typedef TEOSCallback<EOS_UserInfo_OnQueryUserInfoCallback, EOS_UserInfo_QueryUserInfoCallbackInfo, FUserManagerEOS> FReadUserInfoCallback;
 
 void FUserManagerEOS::ReadUserInfo(int32 LocalUserNum, EOS_EpicAccountId EpicAccountId)
 {
-	FReadUserInfoCallback* CallbackObj = new FReadUserInfoCallback();
+	FReadUserInfoCallback* CallbackObj = new FReadUserInfoCallback(FUserManagerEOSWeakPtr(AsShared()));
 	CallbackObj->CallbackLambda = [this, LocalUserNum, EpicAccountId](const EOS_UserInfo_QueryUserInfoCallbackInfo* Data)
 	{
 		if (Data->ResultCode == EOS_EResult::EOS_Success)
@@ -2663,7 +2673,7 @@ struct FQueryByDisplayNameOptions :
 	char DisplayNameAnsi[EOS_OSS_STRING_BUFFER_LENGTH];
 };
 
-typedef TEOSCallback<EOS_UserInfo_OnQueryUserInfoByDisplayNameCallback, EOS_UserInfo_QueryUserInfoByDisplayNameCallbackInfo> FQueryInfoByNameCallback;
+typedef TEOSCallback<EOS_UserInfo_OnQueryUserInfoByDisplayNameCallback, EOS_UserInfo_QueryUserInfoByDisplayNameCallbackInfo, FUserManagerEOS> FQueryInfoByNameCallback;
 
 bool FUserManagerEOS::QueryUserIdMapping(const FUniqueNetId& UserId, const FString& DisplayNameOrEmail, const FOnQueryUserMappingComplete& Delegate)
 {
@@ -2677,7 +2687,7 @@ bool FUserManagerEOS::QueryUserIdMapping(const FUniqueNetId& UserId, const FStri
 	}
 	int32 LocalUserNum = GetLocalUserNumFromUniqueNetId(UserId);
 
-	FQueryInfoByNameCallback* CallbackObj = new FQueryInfoByNameCallback();
+	FQueryInfoByNameCallback* CallbackObj = new FQueryInfoByNameCallback(FUserManagerEOSWeakPtr(AsShared()));
 	CallbackObj->CallbackLambda = [LocalUserNum, DisplayNameOrEmail, this, Delegate](const EOS_UserInfo_QueryUserInfoByDisplayNameCallbackInfo* Data)
 	{
 		EOS_EResult Result = Data->ResultCode;
@@ -2757,7 +2767,7 @@ struct FGetAccountMappingOptions :
 	char AccountId[EOS_CONNECT_EXTERNAL_ACCOUNT_ID_MAX_LENGTH+1];
 };
 
-typedef TEOSCallback<EOS_Connect_OnQueryExternalAccountMappingsCallback, EOS_Connect_QueryExternalAccountMappingsCallbackInfo> FQueryByStringIdsCallback;
+typedef TEOSCallback<EOS_Connect_OnQueryExternalAccountMappingsCallback, EOS_Connect_QueryExternalAccountMappingsCallbackInfo, FUserManagerEOS> FQueryByStringIdsCallback;
 
 bool FUserManagerEOS::QueryExternalIdMappings(const FUniqueNetId& UserId, const FExternalIdQueryOptions& QueryOptions, const TArray<FString>& ExternalIds, const FOnQueryExternalIdMappingsComplete& Delegate)
 {
@@ -2788,7 +2798,7 @@ bool FUserManagerEOS::QueryExternalIdMappings(const FUniqueNetId& UserId, const
 			FCStringAnsi::Strncpy(Options.PointerArray[ProcessedCount], TCHAR_TO_UTF8(*ExternalIds[ProcessedCount]), EOS_CONNECT_EXTERNAL_ACCOUNT_ID_MAX_LENGTH+1);
 			BatchIds.Add(ExternalIds[ProcessedCount]);
 		}
-		FQueryByStringIdsCallback* CallbackObj = new FQueryByStringIdsCallback();
+		FQueryByStringIdsCallback* CallbackObj = new FQueryByStringIdsCallback(FUserManagerEOSWeakPtr(AsShared()));
 		CallbackObj->CallbackLambda = [LocalUserNum, QueryOptions, BatchIds, this, Delegate](const EOS_Connect_QueryExternalAccountMappingsCallbackInfo* Data)
 		{
 			EOS_EResult Result = Data->ResultCode;
diff --git a/Engine/Plugins/Online/OnlineSubsystemEOS/Source/OnlineSubsystemEOS/Private/UserManagerEOS.h b/Engine/Plugins/Online/OnlineSubsystemEOS/Source/OnlineSubsystemEOS/Private/UserManagerEOS.h
--- a/Engine/Plugins/Online/OnlineSubsystemEOS/Source/OnlineSubsystemEOS/Private/UserManagerEOS.h
+++ b/Engine/Plugins/Online/OnlineSubsystemEOS/Source/OnlineSubsystemEOS/Private/UserManagerEOS.h
@@ -241,6 +241,7 @@ class FUserManagerEOS
 	, public IOnlineFriends
 	, public IOnlinePresence
 	, public IOnlineUser
+	, public TSharedFromThis<FUserManagerEOS, ESPMode::ThreadSafe>
 {
 public:
 	/**
@@ -255,6 +256,9 @@ public:
 	 */
 	virtual ~FUserManagerEOS();
 
+	void Init();
+	void Shutdown();
+
 // IOnlineIdentity Interface
 	virtual bool Login(int32 LocalUserNum, const FOnlineAccountCredentials& AccountCredentials) override;
 	virtual bool Logout(int32 LocalUserNum) override;
@@ -481,4 +485,8 @@ private:
 	TMap<int32, TSharedRef<FOnlineAccountCredentials>> LocalUserNumToLastLoginCredentials;
 };
 
+typedef TSharedPtr<FUserManagerEOS, ESPMode::ThreadSafe> FUserManagerEOSPtr;
+typedef TWeakPtr<FUserManagerEOS, ESPMode::ThreadSafe> FUserManagerEOSWeakPtr;
+typedef TWeakPtr<const FUserManagerEOS, ESPMode::ThreadSafe> FUserManagerEOSConstWeakPtr;
+
 #endif
diff --git a/Engine/Plugins/Online/SocketSubsystemEOS/Source/SocketSubsystemEOS/Private/SocketEOS.cpp b/Engine/Plugins/Online/SocketSubsystemEOS/Source/SocketSubsystemEOS/Private/SocketEOS.cpp
--- a/Engine/Plugins/Online/SocketSubsystemEOS/Source/SocketSubsystemEOS/Private/SocketEOS.cpp
+++ b/Engine/Plugins/Online/SocketSubsystemEOS/Source/SocketSubsystemEOS/Private/SocketEOS.cpp
@@ -65,6 +65,7 @@ FSocketEOS::FSocketEOS(FSocketSubsystemEOS& InSocketSubsystem, const FString& In
 	, ClosedNotifyId(EOS_INVALID_NOTIFICATIONID)
 #endif
 {
+	CallbackAliveTracker = MakeShared<FCallbackBase>();
 }
 
 FSocketEOS::~FSocketEOS()
@@ -76,6 +77,7 @@ FSocketEOS::~FSocketEOS()
 		SocketSubsystem.UnbindChannel(LocalAddress);
 		LocalAddress = FInternetAddrEOS();
 	}
+	CallbackAliveTracker = nullptr;
 }
 
 bool FSocketEOS::Shutdown(ESocketShutdownMode Mode)
@@ -201,7 +203,7 @@ bool FSocketEOS::Listen(int32)
 	Options.LocalUserId = LocalAddress.GetLocalUserId();
 	Options.SocketId = &SocketId;
 
-	ConnectNotifyCallback = new FConnectNotifyCallback();
+	ConnectNotifyCallback = new FConnectNotifyCallback(CallbackAliveTracker);
 	ConnectNotifyCallback->CallbackLambda = [this](const EOS_P2P_OnIncomingConnectionRequestInfo* Info)
 	{
 		char PuidBuffer[64];
@@ -662,7 +664,7 @@ void FSocketEOS::RegisterClosedNotification()
 	Options.LocalUserId = LocalAddress.GetLocalUserId();
 	Options.SocketId = &SocketId;
 
-	ClosedNotifyCallback = new FClosedNotifyCallback();
+	ClosedNotifyCallback = new FClosedNotifyCallback(CallbackAliveTracker);
 	ClosedNotifyCallback->CallbackLambda = [this](const EOS_P2P_OnRemoteConnectionClosedInfo* Info)
 	{
 		// Add this connection to the list of closed ones
diff --git a/Engine/Plugins/Online/SocketSubsystemEOS/Source/SocketSubsystemEOS/Public/SocketEOS.h b/Engine/Plugins/Online/SocketSubsystemEOS/Source/SocketSubsystemEOS/Public/SocketEOS.h
--- a/Engine/Plugins/Online/SocketSubsystemEOS/Source/SocketSubsystemEOS/Public/SocketEOS.h
+++ b/Engine/Plugins/Online/SocketSubsystemEOS/Source/SocketSubsystemEOS/Public/SocketEOS.h
@@ -22,9 +22,6 @@ class FSocketSubsystemEOS;
 
 #if WITH_EOS_SDK
 	#include "eos_p2p_types.h"
-
-	typedef TEOSGlobalCallback<EOS_P2P_OnIncomingConnectionRequestCallback, EOS_P2P_OnIncomingConnectionRequestInfo> FConnectNotifyCallback;
-	typedef TEOSGlobalCallback<EOS_P2P_OnRemoteConnectionClosedCallback, EOS_P2P_OnRemoteConnectionClosedInfo> FClosedNotifyCallback;
 #endif
 
 class SOCKETSUBSYSTEMEOS_API FSocketEOS
@@ -83,6 +80,9 @@ public:
 	void RegisterClosedNotification();
 
 private:
+	/** Used to track our aliveness and make it possible to use the callback interface */
+	TSharedPtr<FCallbackBase> CallbackAliveTracker;
+
 	/** Reference to our subsystem */
 	FSocketSubsystemEOS& SocketSubsystem;
 
@@ -95,6 +95,9 @@ private:
 	TArray<FInternetAddrEOS> ClosedRemotes;
 
 #if WITH_EOS_SDK
+	typedef TEOSGlobalCallback<EOS_P2P_OnIncomingConnectionRequestCallback, EOS_P2P_OnIncomingConnectionRequestInfo, FCallbackBase> FConnectNotifyCallback;
+	typedef TEOSGlobalCallback<EOS_P2P_OnRemoteConnectionClosedCallback, EOS_P2P_OnRemoteConnectionClosedInfo, FCallbackBase> FClosedNotifyCallback;
+
 	FConnectNotifyCallback* ConnectNotifyCallback;
 	EOS_NotificationId ConnectNotifyId;
 
