Added validity checks for Friends interface, since it can be null in some configurations

diff --git a/Engine/Plugins/Online/OnlineSubsystemEOS/Source/OnlineSubsystemEOSPlus/Private/OnlineUserEOSPlus.cpp b/Engine/Plugins/Online/OnlineSubsystemEOS/Source/OnlineSubsystemEOSPlus/Private/OnlineUserEOSPlus.cpp
--- a/Engine/Plugins/Online/OnlineSubsystemEOS/Source/OnlineSubsystemEOSPlus/Private/OnlineUserEOSPlus.cpp
+++ b/Engine/Plugins/Online/OnlineSubsystemEOS/Source/OnlineSubsystemEOSPlus/Private/OnlineUserEOSPlus.cpp
@@ -134,26 +134,34 @@ bool FUniqueNetIdEOSPlus::IsValid() const
 FOnlineUserEOSPlus::FOnlineUserEOSPlus(FOnlineSubsystemEOSPlus* InSubsystem)
 	: EOSPlus(InSubsystem)
 {
-	BaseUserInterface = EOSPlus->BaseOSS->GetUserInterface(); //We don't check it here, since some platforms might not implement it
+	// Some interfaces, we don't check() here, since some platforms might not implement them
+	BaseUserInterface = EOSPlus->BaseOSS->GetUserInterface();
+	BaseFriendsInterface = EOSPlus->BaseOSS->GetFriendsInterface();
+	if (!BaseFriendsInterface.IsValid())
+	{
+		UE_LOG_ONLINE(Log, TEXT("[FOnlineUserEOSPlus::Initialize] BaseFriendsInterface not valid. Delegates will not be bound and methods will fall back on EOS functionality only."));
+	}
 
+	EOSFriendsInterface = EOSPlus->EosOSS->GetFriendsInterface();
+	check(EOSFriendsInterface.IsValid());
 	BaseIdentityInterface = EOSPlus->BaseOSS->GetIdentityInterface();
 	check(BaseIdentityInterface.IsValid());
 	EOSIdentityInterface = EOSPlus->EosOSS->GetIdentityInterface();
 	check(EOSIdentityInterface.IsValid());
-	BaseFriendsInterface = EOSPlus->BaseOSS->GetFriendsInterface();
-	check(BaseFriendsInterface.IsValid());
-	EOSFriendsInterface = EOSPlus->EosOSS->GetFriendsInterface();
-	check(EOSFriendsInterface.IsValid());
 	BasePresenceInterface = EOSPlus->BaseOSS->GetPresenceInterface();
 	check(BasePresenceInterface.IsValid());
 	EOSPresenceInterface = EOSPlus->EosOSS->GetPresenceInterface();
 	check(EOSPresenceInterface.IsValid());
 
-	BaseFriendsInterface->AddOnInviteReceivedDelegate_Handle(FOnInviteReceivedDelegate::CreateRaw(this, &FOnlineUserEOSPlus::OnInviteReceived));
-	BaseFriendsInterface->AddOnInviteAcceptedDelegate_Handle(FOnInviteAcceptedDelegate::CreateRaw(this, &FOnlineUserEOSPlus::OnInviteAccepted));
-	BaseFriendsInterface->AddOnInviteRejectedDelegate_Handle(FOnInviteRejectedDelegate::CreateRaw(this, &FOnlineUserEOSPlus::OnInviteRejected));
-	BaseFriendsInterface->AddOnInviteAbortedDelegate_Handle(FOnInviteAbortedDelegate::CreateRaw(this, &FOnlineUserEOSPlus::OnInviteAborted));
-	BaseFriendsInterface->AddOnFriendRemovedDelegate_Handle(FOnFriendRemovedDelegate::CreateRaw(this, &FOnlineUserEOSPlus::OnFriendRemoved));
+	if (BaseFriendsInterface.IsValid())
+	{
+		BaseFriendsInterface->AddOnInviteReceivedDelegate_Handle(FOnInviteReceivedDelegate::CreateRaw(this, &FOnlineUserEOSPlus::OnInviteReceived));
+		BaseFriendsInterface->AddOnInviteAcceptedDelegate_Handle(FOnInviteAcceptedDelegate::CreateRaw(this, &FOnlineUserEOSPlus::OnInviteAccepted));
+		BaseFriendsInterface->AddOnInviteRejectedDelegate_Handle(FOnInviteRejectedDelegate::CreateRaw(this, &FOnlineUserEOSPlus::OnInviteRejected));
+		BaseFriendsInterface->AddOnInviteAbortedDelegate_Handle(FOnInviteAbortedDelegate::CreateRaw(this, &FOnlineUserEOSPlus::OnInviteAborted));
+		BaseFriendsInterface->AddOnFriendRemovedDelegate_Handle(FOnFriendRemovedDelegate::CreateRaw(this, &FOnlineUserEOSPlus::OnFriendRemoved));
+	}
+
 	EOSFriendsInterface->AddOnInviteReceivedDelegate_Handle(FOnInviteReceivedDelegate::CreateRaw(this, &FOnlineUserEOSPlus::OnInviteReceived));
 	EOSFriendsInterface->AddOnInviteAcceptedDelegate_Handle(FOnInviteAcceptedDelegate::CreateRaw(this, &FOnlineUserEOSPlus::OnInviteAccepted));
 	EOSFriendsInterface->AddOnInviteRejectedDelegate_Handle(FOnInviteRejectedDelegate::CreateRaw(this, &FOnlineUserEOSPlus::OnInviteRejected));
@@ -172,9 +180,13 @@ FOnlineUserEOSPlus::FOnlineUserEOSPlus(FOnlineSubsystemEOSPlus* InSubsystem)
 		BaseIdentityInterface->AddOnLoginCompleteDelegate_Handle(LocalUserNum, FOnLoginCompleteDelegate::CreateRaw(this, &FOnlineUserEOSPlus::OnBaseLoginComplete));
 		BaseIdentityInterface->AddOnLogoutCompleteDelegate_Handle(LocalUserNum, FOnLogoutCompleteDelegate::CreateRaw(this, &FOnlineUserEOSPlus::OnLogoutComplete));
 
-		BaseFriendsInterface->AddOnFriendsChangeDelegate_Handle(LocalUserNum, FOnFriendsChangeDelegate::CreateRaw(this, &FOnlineUserEOSPlus::OnFriendsChanged));
+		if (BaseFriendsInterface.IsValid())
+		{
+			BaseFriendsInterface->AddOnFriendsChangeDelegate_Handle(LocalUserNum, FOnFriendsChangeDelegate::CreateRaw(this, &FOnlineUserEOSPlus::OnFriendsChanged));
+			BaseFriendsInterface->AddOnOutgoingInviteSentDelegate_Handle(LocalUserNum, FOnOutgoingInviteSentDelegate::CreateRaw(this, &FOnlineUserEOSPlus::OnOutgoingInviteSent));
+		}
+
 		EOSFriendsInterface->AddOnFriendsChangeDelegate_Handle(LocalUserNum, FOnFriendsChangeDelegate::CreateRaw(this, &FOnlineUserEOSPlus::OnFriendsChanged));
-		BaseFriendsInterface->AddOnOutgoingInviteSentDelegate_Handle(LocalUserNum, FOnOutgoingInviteSentDelegate::CreateRaw(this, &FOnlineUserEOSPlus::OnOutgoingInviteSent));
 		EOSFriendsInterface->AddOnOutgoingInviteSentDelegate_Handle(LocalUserNum, FOnOutgoingInviteSentDelegate::CreateRaw(this, &FOnlineUserEOSPlus::OnOutgoingInviteSent));
 	}
 }
@@ -183,11 +195,14 @@ FOnlineUserEOSPlus::~FOnlineUserEOSPlus()
 {
 	BaseIdentityInterface->ClearOnLoginChangedDelegates(this);
 	BaseIdentityInterface->ClearOnControllerPairingChangedDelegates(this);
-	BaseFriendsInterface->ClearOnInviteReceivedDelegates(this);
-	BaseFriendsInterface->ClearOnInviteAcceptedDelegates(this);
-	BaseFriendsInterface->ClearOnInviteRejectedDelegates(this);
-	BaseFriendsInterface->ClearOnInviteAbortedDelegates(this);
-	BaseFriendsInterface->ClearOnFriendRemovedDelegates(this);
+	if (BaseFriendsInterface.IsValid())
+	{
+		BaseFriendsInterface->ClearOnInviteReceivedDelegates(this);
+		BaseFriendsInterface->ClearOnInviteAcceptedDelegates(this);
+		BaseFriendsInterface->ClearOnInviteRejectedDelegates(this);
+		BaseFriendsInterface->ClearOnInviteAbortedDelegates(this);
+		BaseFriendsInterface->ClearOnFriendRemovedDelegates(this);
+	}
 	EOSFriendsInterface->ClearOnInviteReceivedDelegates(this);
 	EOSFriendsInterface->ClearOnInviteAcceptedDelegates(this);
 	EOSFriendsInterface->ClearOnInviteRejectedDelegates(this);
@@ -207,9 +222,12 @@ FOnlineUserEOSPlus::~FOnlineUserEOSPlus()
 		BaseIdentityInterface->ClearOnLoginCompleteDelegates(LocalUserNum, this);
 		BaseIdentityInterface->ClearOnLogoutCompleteDelegates(LocalUserNum, this);
 
-		BaseFriendsInterface->ClearOnFriendsChangeDelegates(LocalUserNum, this);
+		if (BaseFriendsInterface.IsValid())
+		{
+			BaseFriendsInterface->ClearOnFriendsChangeDelegates(LocalUserNum, this);
+			BaseFriendsInterface->ClearOnOutgoingInviteSentDelegates(LocalUserNum, this);
+		}
 		EOSFriendsInterface->ClearOnFriendsChangeDelegates(LocalUserNum, this);
-		BaseFriendsInterface->ClearOnOutgoingInviteSentDelegates(LocalUserNum, this);
 		EOSFriendsInterface->ClearOnOutgoingInviteSentDelegates(LocalUserNum, this);
 	}
 }
@@ -917,27 +935,56 @@ void FOnlineUserEOSPlus::OnFriendRemoved(const FUniqueNetId& UserId, const FUniq
 
 bool FOnlineUserEOSPlus::ReadFriendsList(int32 LocalUserNum, const FString& ListName, const FOnReadFriendsListComplete& Delegate)
 {
-	return BaseFriendsInterface->ReadFriendsList(LocalUserNum, ListName,
-		FOnReadFriendsListComplete::CreateLambda([this, IntermediateComplete = FOnReadFriendsListComplete(Delegate)](int32 LocalUserNum, bool bWasSuccessful, const FString& ListName, const FString& ErrorStr)
+	if (BaseFriendsInterface.IsValid())
 	{
-		// Skip reading EAS if not in use and if we errored at the platform level
-		if (!UEOSSettings::GetSettings().bUseEAS || !bWasSuccessful)
+		return BaseFriendsInterface->ReadFriendsList(LocalUserNum, ListName,
+			FOnReadFriendsListComplete::CreateLambda([this, IntermediateComplete = FOnReadFriendsListComplete(Delegate)](int32 LocalUserNum, bool bWasSuccessful, const FString& ListName, const FString& ErrorStr)
 		{
-			IntermediateComplete.ExecuteIfBound(LocalUserNum, bWasSuccessful, ListName, ErrorStr);
-			return;
+			// Skip reading EAS if not in use and if we errored at the platform level
+			if (!UEOSSettings::GetSettings().bUseEAS || !bWasSuccessful)
+			{
+				IntermediateComplete.ExecuteIfBound(LocalUserNum, bWasSuccessful, ListName, ErrorStr);
+				return;
+			}
+			// Read the EAS version too
+			EOSFriendsInterface->ReadFriendsList(LocalUserNum, ListName,
+				FOnReadFriendsListComplete::CreateLambda([this, OnComplete = FOnReadFriendsListComplete(IntermediateComplete)](int32 LocalUserNum, bool bWasSuccessful, const FString& ListName, const FString& ErrorStr)
+			{
+				OnComplete.ExecuteIfBound(LocalUserNum, bWasSuccessful, ListName, ErrorStr);
+			}));
+		}));
+	}
+	else
+	{
+		UE_LOG_ONLINE(Verbose, TEXT("[FOnlineUserEOSPlus::ReadFriendsList] BaseFriendsInterface not valid"));
+
+		// Skip reading EAS if not in use
+		if (!UEOSSettings::GetSettings().bUseEAS)
+		{
+			Delegate.ExecuteIfBound(LocalUserNum, false, ListName, FString());
+			return false;
 		}
-		// Read the EAS version too
-		EOSFriendsInterface->ReadFriendsList(LocalUserNum, ListName,
-			FOnReadFriendsListComplete::CreateLambda([this, OnComplete = FOnReadFriendsListComplete(IntermediateComplete)](int32 LocalUserNum, bool bWasSuccessful, const FString& ListName, const FString& ErrorStr)
+
+		// Read the EAS version
+		return EOSFriendsInterface->ReadFriendsList(LocalUserNum, ListName,
+			FOnReadFriendsListComplete::CreateLambda([this, OnComplete = FOnReadFriendsListComplete(Delegate)](int32 LocalUserNum, bool bWasSuccessful, const FString& ListName, const FString& ErrorStr)
 		{
 			OnComplete.ExecuteIfBound(LocalUserNum, bWasSuccessful, ListName, ErrorStr);
 		}));
-	}));
+	}
 }
 
 bool FOnlineUserEOSPlus::DeleteFriendsList(int32 LocalUserNum, const FString& ListName, const FOnDeleteFriendsListComplete& Delegate)
 {
-	return BaseFriendsInterface->DeleteFriendsList(LocalUserNum, ListName, Delegate);
+	if (BaseFriendsInterface.IsValid())
+	{
+		return BaseFriendsInterface->DeleteFriendsList(LocalUserNum, ListName, Delegate);
+	}
+	else
+	{
+		UE_LOG_ONLINE(Verbose, TEXT("[FOnlineUserEOSPlus::DeleteFriendsList] BaseFriendsInterface not valid"));
+		return false;
+	}
 }
 
 bool FOnlineUserEOSPlus::SendInvite(int32 LocalUserNum, const FUniqueNetId& FriendId, const FString& ListName, const FOnSendInviteComplete& Delegate)
@@ -946,7 +993,16 @@ bool FOnlineUserEOSPlus::SendInvite(int32 LocalUserNum, const FUniqueNetId& Frie
 	{
 		return false;
 	}
-	return BaseFriendsInterface->SendInvite(LocalUserNum, *NetIdPlusToBaseNetId[FriendId.ToString()], ListName, Delegate);
+
+	if (BaseFriendsInterface.IsValid())
+	{
+		return BaseFriendsInterface->SendInvite(LocalUserNum, *NetIdPlusToBaseNetId[FriendId.ToString()], ListName, Delegate);
+	}
+	else
+	{
+		UE_LOG_ONLINE(Verbose, TEXT("[FOnlineUserEOSPlus::SendInvite] BaseFriendsInterface not valid"));
+		return false;
+	}
 }
 
 bool FOnlineUserEOSPlus::AcceptInvite(int32 LocalUserNum, const FUniqueNetId& FriendId, const FString& ListName, const FOnAcceptInviteComplete& Delegate)
@@ -955,7 +1011,16 @@ bool FOnlineUserEOSPlus::AcceptInvite(int32 LocalUserNum, const FUniqueNetId& Fr
 	{
 		return false;
 	}
-	return BaseFriendsInterface->AcceptInvite(LocalUserNum, *NetIdPlusToBaseNetId[FriendId.ToString()], ListName, Delegate);
+
+	if (BaseFriendsInterface.IsValid())
+	{
+		return BaseFriendsInterface->AcceptInvite(LocalUserNum, *NetIdPlusToBaseNetId[FriendId.ToString()], ListName, Delegate);
+	}
+	else
+	{
+		UE_LOG_ONLINE(Verbose, TEXT("[FOnlineUserEOSPlus::AcceptInvite] BaseFriendsInterface not valid"));
+		return false;
+	}
 }
 
 bool FOnlineUserEOSPlus::RejectInvite(int32 LocalUserNum, const FUniqueNetId& FriendId, const FString& ListName)
@@ -964,7 +1029,16 @@ bool FOnlineUserEOSPlus::RejectInvite(int32 LocalUserNum, const FUniqueNetId& Fr
 	{
 		return false;
 	}
-	return BaseFriendsInterface->RejectInvite(LocalUserNum, *NetIdPlusToBaseNetId[FriendId.ToString()], ListName);
+
+	if (BaseFriendsInterface.IsValid())
+	{
+		return BaseFriendsInterface->RejectInvite(LocalUserNum, *NetIdPlusToBaseNetId[FriendId.ToString()], ListName);
+	}
+	else
+	{
+		UE_LOG_ONLINE(Verbose, TEXT("[FOnlineUserEOSPlus::RejectInvite] BaseFriendsInterface not valid"));
+		return false;
+	}
 }
 
 bool FOnlineUserEOSPlus::DeleteFriend(int32 LocalUserNum, const FUniqueNetId& FriendId, const FString& ListName)
@@ -973,7 +1047,16 @@ bool FOnlineUserEOSPlus::DeleteFriend(int32 LocalUserNum, const FUniqueNetId& Fr
 	{
 		return false;
 	}
-	return BaseFriendsInterface->DeleteFriend(LocalUserNum, *NetIdPlusToBaseNetId[FriendId.ToString()], ListName);
+
+	if (BaseFriendsInterface.IsValid())
+	{
+		return BaseFriendsInterface->DeleteFriend(LocalUserNum, *NetIdPlusToBaseNetId[FriendId.ToString()], ListName);
+	}
+	else
+	{
+		UE_LOG_ONLINE(Verbose, TEXT("[FOnlineUserEOSPlus::DeleteFriend] BaseFriendsInterface not valid"));
+		return false;
+	}
 }
 
 TSharedRef<FOnlineFriendPlus> FOnlineUserEOSPlus::AddFriend(TSharedRef<FOnlineFriend> Friend)
@@ -1038,11 +1121,21 @@ bool FOnlineUserEOSPlus::GetFriendsList(int32 LocalUserNum, const FString& ListN
 	OutFriends.Reset();
 
 	TArray<TSharedRef<FOnlineFriend>> Friends;
-	bool bWasSuccessful = BaseFriendsInterface->GetFriendsList(LocalUserNum, ListName, Friends);
-	// Build the list of base friends
-	for (TSharedRef<FOnlineFriend> Friend : Friends)
+
+	bool bWasSuccessful = false;
+
+	if (BaseFriendsInterface.IsValid())
 	{
-		OutFriends.Add(GetFriend(Friend));
+		bWasSuccessful = BaseFriendsInterface->GetFriendsList(LocalUserNum, ListName, Friends);
+		// Build the list of base friends
+		for (TSharedRef<FOnlineFriend> Friend : Friends)
+		{
+			OutFriends.Add(GetFriend(Friend));
+		}
+	}
+	else
+	{
+		UE_LOG_ONLINE(Verbose, TEXT("[FOnlineUserEOSPlus::GetFriendsList] BaseFriendsInterface not valid"));
 	}
 
 	if (UEOSSettings::GetSettings().bUseEAS)
@@ -1055,6 +1148,7 @@ bool FOnlineUserEOSPlus::GetFriendsList(int32 LocalUserNum, const FString& ListN
 			OutFriends.Add(GetFriend(Friend));
 		}
 	}
+
 	return bWasSuccessful;
 }
 
@@ -1083,10 +1177,19 @@ TSharedPtr<FOnlineFriend> FOnlineUserEOSPlus::GetFriend(int32 LocalUserNum, cons
 bool FOnlineUserEOSPlus::IsFriend(int32 LocalUserNum, const FUniqueNetId& FriendId, const FString& ListName)
 {
 	bool bIsFriend = false;
-	if (NetIdPlusToBaseNetId.Contains(FriendId.ToString()))
+
+	if (BaseFriendsInterface.IsValid())
 	{
-		bIsFriend = BaseFriendsInterface->IsFriend(LocalUserNum, *NetIdPlusToBaseNetId[FriendId.ToString()], ListName);
+		if (NetIdPlusToBaseNetId.Contains(FriendId.ToString()))
+		{
+			bIsFriend = BaseFriendsInterface->IsFriend(LocalUserNum, *NetIdPlusToBaseNetId[FriendId.ToString()], ListName);
+		}
+	}
+	else
+	{
+		UE_LOG_ONLINE(Verbose, TEXT("[FOnlineUserEOSPlus::IsFriend] BaseFriendsInterface not valid"));
 	}
+
 	if (!bIsFriend && UEOSSettings::GetSettings().bUseEAS && NetIdPlusToEOSNetId.Contains(FriendId.ToString()))
 	{
 		bIsFriend = EOSFriendsInterface->IsFriend(LocalUserNum, *NetIdPlusToEOSNetId[FriendId.ToString()], ListName);
@@ -1100,7 +1203,16 @@ bool FOnlineUserEOSPlus::QueryRecentPlayers(const FUniqueNetId& UserId, const FS
 	{
 		return false;
 	}
-	return BaseFriendsInterface->QueryRecentPlayers(*NetIdPlusToBaseNetId[UserId.ToString()], Namespace);
+
+	if (BaseFriendsInterface.IsValid())
+	{
+		return BaseFriendsInterface->QueryRecentPlayers(*NetIdPlusToBaseNetId[UserId.ToString()], Namespace);
+	}
+	else
+	{
+		UE_LOG_ONLINE(Verbose, TEXT("[FOnlineUserEOSPlus::QueryRecentPlayers] BaseFriendsInterface not valid"));
+		return false;
+	}
 }
 
 TSharedRef<FOnlineRecentPlayer> FOnlineUserEOSPlus::AddRecentPlayer(TSharedRef<FOnlineRecentPlayer> Player)
@@ -1159,11 +1271,20 @@ bool FOnlineUserEOSPlus::GetRecentPlayers(const FUniqueNetId& UserId, const FStr
 		return false;
 	}
 
-	TArray<TSharedRef<FOnlineRecentPlayer>> Players;
-	bool bWasSuccessful = BaseFriendsInterface->GetRecentPlayers(*NetIdPlusToBaseNetId[UserId.ToString()], Namespace, Players);
-	for (TSharedRef<FOnlineRecentPlayer> Player : Players)
+	bool bWasSuccessful = false;
+
+	if (BaseFriendsInterface.IsValid())
+	{
+		TArray<TSharedRef<FOnlineRecentPlayer>> Players;
+		bWasSuccessful = BaseFriendsInterface->GetRecentPlayers(*NetIdPlusToBaseNetId[UserId.ToString()], Namespace, Players);
+		for (TSharedRef<FOnlineRecentPlayer> Player : Players)
+		{
+			OutRecentPlayers.Add(GetRecentPlayer(Player));
+		}
+	}
+	else
 	{
-		OutRecentPlayers.Add(GetRecentPlayer(Player));
+		UE_LOG_ONLINE(Verbose, TEXT("[FOnlineUserEOSPlus::GetRecentPlayers] BaseFriendsInterface not valid"));
 	}
 
 	return bWasSuccessful;
@@ -1175,7 +1296,16 @@ bool FOnlineUserEOSPlus::BlockPlayer(int32 LocalUserNum, const FUniqueNetId& Pla
 	{
 		return false;
 	}
-	return BaseFriendsInterface->BlockPlayer(LocalUserNum, *NetIdPlusToBaseNetId[PlayerId.ToString()]);
+
+	if (BaseFriendsInterface.IsValid())
+	{
+		return BaseFriendsInterface->BlockPlayer(LocalUserNum, *NetIdPlusToBaseNetId[PlayerId.ToString()]);
+	}
+	else
+	{
+		UE_LOG_ONLINE(Verbose, TEXT("[FOnlineUserEOSPlus::BlockPlayer] BaseFriendsInterface not valid"));
+		return false;
+	}
 }
 
 bool FOnlineUserEOSPlus::UnblockPlayer(int32 LocalUserNum, const FUniqueNetId& PlayerId)
@@ -1184,7 +1314,16 @@ bool FOnlineUserEOSPlus::UnblockPlayer(int32 LocalUserNum, const FUniqueNetId& P
 	{
 		return false;
 	}
-	return BaseFriendsInterface->UnblockPlayer(LocalUserNum, *NetIdPlusToBaseNetId[PlayerId.ToString()]);
+
+	if (BaseFriendsInterface.IsValid())
+	{
+		return BaseFriendsInterface->UnblockPlayer(LocalUserNum, *NetIdPlusToBaseNetId[PlayerId.ToString()]);
+	}
+	else
+	{
+		UE_LOG_ONLINE(Verbose, TEXT("[FOnlineUserEOSPlus::UnblockPlayer] BaseFriendsInterface not valid"));
+		return false;
+	}
 }
 
 bool FOnlineUserEOSPlus::QueryBlockedPlayers(const FUniqueNetId& UserId)
@@ -1193,7 +1332,16 @@ bool FOnlineUserEOSPlus::QueryBlockedPlayers(const FUniqueNetId& UserId)
 	{
 		return false;
 	}
-	return BaseFriendsInterface->QueryBlockedPlayers(*NetIdPlusToBaseNetId[UserId.ToString()]);
+
+	if (BaseFriendsInterface.IsValid())
+	{
+		return BaseFriendsInterface->QueryBlockedPlayers(*NetIdPlusToBaseNetId[UserId.ToString()]);
+	}
+	else
+	{
+		UE_LOG_ONLINE(Verbose, TEXT("[FOnlineUserEOSPlus::QueryBlockedPlayers] BaseFriendsInterface not valid"));
+		return false;
+	}
 }
 
 TSharedRef<FOnlineBlockedPlayer> FOnlineUserEOSPlus::AddBlockedPlayer(TSharedRef<FOnlineBlockedPlayer> Player)
@@ -1248,28 +1396,54 @@ bool FOnlineUserEOSPlus::GetBlockedPlayers(const FUniqueNetId& UserId, TArray< T
 		return false;
 	}
 
-	TArray<TSharedRef<FOnlineBlockedPlayer>> Players;
-	bool bWasSuccessful = BaseFriendsInterface->GetBlockedPlayers(*NetIdPlusToBaseNetId[UserId.ToString()], Players);
-	for (TSharedRef<FOnlineBlockedPlayer> Player : Players)
+	bool bWasSuccessful = false;
+
+	if (BaseFriendsInterface.IsValid())
+	{
+		TArray<TSharedRef<FOnlineBlockedPlayer>> Players;
+		bWasSuccessful = BaseFriendsInterface->GetBlockedPlayers(*NetIdPlusToBaseNetId[UserId.ToString()], Players);
+		for (TSharedRef<FOnlineBlockedPlayer> Player : Players)
+		{
+			OutBlockedPlayers.Add(GetBlockedPlayer(Player));
+		}
+	}
+	else
 	{
-		OutBlockedPlayers.Add(GetBlockedPlayer(Player));
+		UE_LOG_ONLINE(Verbose, TEXT("[FOnlineUserEOSPlus::GetBlockedPlayers] BaseFriendsInterface not valid"));
 	}
+
 	return bWasSuccessful;
 }
 
 void FOnlineUserEOSPlus::DumpBlockedPlayers() const
 {
-	BaseFriendsInterface->DumpBlockedPlayers();
+	if (BaseFriendsInterface.IsValid())
+	{
+		BaseFriendsInterface->DumpBlockedPlayers();
+	}
+	else
+	{
+		UE_LOG_ONLINE(Verbose, TEXT("[FOnlineUserEOSPlus::DumpBlockedPlayers] BaseFriendsInterface not valid"));
+	}
+
 	EOSFriendsInterface->DumpBlockedPlayers();
 }
 
 void FOnlineUserEOSPlus::SetFriendAlias(int32 LocalUserNum, const FUniqueNetId& FriendId, const FString& ListName, const FString& Alias, const FOnSetFriendAliasComplete& Delegate)
 {
-	if (NetIdPlusToBaseNetId.Contains(FriendId.ToString()))
+	if (BaseFriendsInterface.IsValid())
 	{
-		BaseFriendsInterface->SetFriendAlias(LocalUserNum, *NetIdPlusToBaseNetId[FriendId.ToString()], ListName, Alias, Delegate);
-		return;
+		if (NetIdPlusToBaseNetId.Contains(FriendId.ToString()))
+		{
+			BaseFriendsInterface->SetFriendAlias(LocalUserNum, *NetIdPlusToBaseNetId[FriendId.ToString()], ListName, Alias, Delegate);
+			return;
+		}
 	}
+	else
+	{
+		UE_LOG_ONLINE(Verbose, TEXT("[FOnlineUserEOSPlus::SetFriendAlias] BaseFriendsInterface not valid"));
+	}
+
 	if (NetIdPlusToEOSNetId.Contains(FriendId.ToString()))
 	{
 		EOSFriendsInterface->SetFriendAlias(LocalUserNum, *NetIdPlusToEOSNetId[FriendId.ToString()], ListName, Alias, Delegate);
@@ -1280,11 +1454,19 @@ void FOnlineUserEOSPlus::SetFriendAlias(int32 LocalUserNum, const FUniqueNetId&
 
 void FOnlineUserEOSPlus::DeleteFriendAlias(int32 LocalUserNum, const FUniqueNetId& FriendId, const FString& ListName, const FOnDeleteFriendAliasComplete& Delegate)
 {
-	if (NetIdPlusToBaseNetId.Contains(FriendId.ToString()))
+	if (BaseFriendsInterface.IsValid())
 	{
-		BaseFriendsInterface->DeleteFriendAlias(LocalUserNum, *NetIdPlusToBaseNetId[FriendId.ToString()], ListName, Delegate);
-		return;
+		if (NetIdPlusToBaseNetId.Contains(FriendId.ToString()))
+		{
+			BaseFriendsInterface->DeleteFriendAlias(LocalUserNum, *NetIdPlusToBaseNetId[FriendId.ToString()], ListName, Delegate);
+			return;
+		}
+	}
+	else
+	{
+		UE_LOG_ONLINE(Verbose, TEXT("[FOnlineUserEOSPlus::DeleteFriendAlias] BaseFriendsInterface not valid"));
 	}
+
 	if (NetIdPlusToEOSNetId.Contains(FriendId.ToString()))
 	{
 		EOSFriendsInterface->DeleteFriendAlias(LocalUserNum, *NetIdPlusToEOSNetId[FriendId.ToString()], ListName, Delegate);
@@ -1295,7 +1477,15 @@ void FOnlineUserEOSPlus::DeleteFriendAlias(int32 LocalUserNum, const FUniqueNetI
 
 void FOnlineUserEOSPlus::DumpRecentPlayers() const
 {
-	BaseFriendsInterface->DumpRecentPlayers();
+	if (BaseFriendsInterface.IsValid())
+	{
+		BaseFriendsInterface->DumpRecentPlayers();
+	}
+	else
+	{
+		UE_LOG_ONLINE(Verbose, TEXT("[FOnlineUserEOSPlus::DumpRecentPlayers] BaseFriendsInterface not valid"));
+	}
+
 	EOSFriendsInterface->DumpRecentPlayers();
 }
 
