- Add accessor for EOSPlatformHandle

diff --git a/Engine/Plugins/Online/OnlineSubsystemEOS/Source/OnlineSubsystemEOS/Public/OnlineSubsystemEOS.h b/Engine/Plugins/Online/OnlineSubsystemEOS/Source/OnlineSubsystemEOS/Public/OnlineSubsystemEOS.h
--- a/Engine/Plugins/Online/OnlineSubsystemEOS/Source/OnlineSubsystemEOS/Public/OnlineSubsystemEOS.h
+++ b/Engine/Plugins/Online/OnlineSubsystemEOS/Source/OnlineSubsystemEOS/Public/OnlineSubsystemEOS.h
@@ -105,6 +105,7 @@ public:
 	virtual bool Tick(float DeltaTime) override;
 
 	IVoiceChatUser* GetVoiceChatUserInterface(const FUniqueNetId& LocalUserId);
+	IEOSPlatformHandlePtr GetEOSPlatformHandle() { return EOSPlatformHandle; };
 
 PACKAGE_SCOPE:
 	/** Only the factory makes instances */
