Addition of UserLoginInfo field to Connect Login options to process auth correctly in platforms that require it

diff --git a/Engine/Plugins/Online/OnlineSubsystemEOS/Source/OnlineSubsystemEOS/OnlineSubsystemEOS.Build.cs b/Engine/Plugins/Online/OnlineSubsystemEOS/Source/OnlineSubsystemEOS/OnlineSubsystemEOS.Build.cs
--- a/Engine/Plugins/Online/OnlineSubsystemEOS/Source/OnlineSubsystemEOS/OnlineSubsystemEOS.Build.cs
+++ b/Engine/Plugins/Online/OnlineSubsystemEOS/Source/OnlineSubsystemEOS/OnlineSubsystemEOS.Build.cs
@@ -35,8 +35,10 @@ public class OnlineSubsystemEOS : ModuleRules
 
 		PrivateDefinitions.Add("USE_XBL_XSTS_TOKEN=" + (bUseXblXstsToken ? "1" : "0"));
 		PrivateDefinitions.Add("USE_PSN_ID_TOKEN=" + (bUsePsnIdToken ? "1" : "0"));
+		PrivateDefinitions.Add("ADD_USER_LOGIN_INFO=" + (bAddUserLoginInfo ? "1" : "0"));
 	}
 
 	protected virtual bool bUseXblXstsToken { get { return false; } }
 	protected virtual bool bUsePsnIdToken { get { return false; } }
+	protected virtual bool bAddUserLoginInfo { get { return false; } }
 }
diff --git a/Engine/Plugins/Online/OnlineSubsystemEOS/Source/OnlineSubsystemEOS/Private/UserManagerEOS.cpp b/Engine/Plugins/Online/OnlineSubsystemEOS/Source/OnlineSubsystemEOS/Private/UserManagerEOS.cpp
--- a/Engine/Plugins/Online/OnlineSubsystemEOS/Source/OnlineSubsystemEOS/Private/UserManagerEOS.cpp
+++ b/Engine/Plugins/Online/OnlineSubsystemEOS/Source/OnlineSubsystemEOS/Private/UserManagerEOS.cpp
@@ -233,7 +233,7 @@ void FUserManagerEOS::LoginStatusChanged(const EOS_Auth_LoginStatusChangedCallba
 	}
 }
 
-IOnlineSubsystem* FUserManagerEOS::GetPlatformOSS()
+IOnlineSubsystem* FUserManagerEOS::GetPlatformOSS() const
 {
 	IOnlineSubsystem* PlatformOSS = IOnlineSubsystem::GetByPlatform();
 	if (PlatformOSS == nullptr)
@@ -250,7 +250,7 @@ IOnlineSubsystem* FUserManagerEOS::GetPlatformOSS()
 	return PlatformOSS;
 }
 
-void FUserManagerEOS::GetPlatformAuthToken(int32 LocalUserNum, const FOnGetLinkedAccountAuthTokenCompleteDelegate& Delegate)
+void FUserManagerEOS::GetPlatformAuthToken(int32 LocalUserNum, const FOnGetLinkedAccountAuthTokenCompleteDelegate& Delegate) const
 {
 	IOnlineSubsystem* PlatformOSS = GetPlatformOSS();
 	if (PlatformOSS == nullptr)
@@ -270,6 +270,28 @@ void FUserManagerEOS::GetPlatformAuthToken(int32 LocalUserNum, const FOnGetLinke
 	PlatformIdentity->GetLinkedAccountAuthToken(LocalUserNum, Delegate);
 }
 
+FString FUserManagerEOS::GetPlatformDisplayName(int32 LocalUserNum) const
+{
+	FString Result;
+
+	IOnlineSubsystem* PlatformOSS = GetPlatformOSS();
+	if (PlatformOSS == nullptr)
+	{
+		UE_LOG_ONLINE(Warning, TEXT("GetPlatformDisplayName(%d) failed due to no platform OSS"), LocalUserNum);
+		return Result;
+	}
+	IOnlineIdentityPtr PlatformIdentity = PlatformOSS->GetIdentityInterface();
+	if (!PlatformIdentity.IsValid())
+	{
+		UE_LOG_ONLINE(Warning, TEXT("GetPlatformDisplayName(%d) failed due to no platform OSS identity interface"), LocalUserNum);
+		return Result;
+	}
+
+	Result = PlatformIdentity->GetPlayerNickname(LocalUserNum);
+
+	return Result;
+}
+
 typedef TEOSCallback<EOS_Auth_OnLoginCallback, EOS_Auth_LoginCallbackInfo, FUserManagerEOS> FLoginCallback;
 typedef TEOSCallback<EOS_Connect_OnLoginCallback, EOS_Connect_LoginCallbackInfo, FUserManagerEOS> FConnectLoginCallback;
 typedef TEOSCallback<EOS_Auth_OnDeletePersistentAuthCallback, EOS_Auth_DeletePersistentAuthCallbackInfo, FUserManagerEOS> FDeletePersistentAuthCallback;
@@ -627,13 +649,21 @@ bool FUserManagerEOS::ConnectLoginNoEAS(int32 LocalUserNum)
 			}
 
 			// Now login into our EOS account
-
 			check(LocalUserNumToLastLoginCredentials.Contains(LocalUserNum));
 			FConnectCredentials Credentials(ToEOS_EExternalCredentialType(GetPlatformOSS()->GetSubsystemName(), *LocalUserNumToLastLoginCredentials[LocalUserNum]), AuthToken);
 			EOS_Connect_LoginOptions Options = { };
 			Options.ApiVersion = EOS_CONNECT_LOGIN_API_LATEST;
 			Options.Credentials = &Credentials;
 
+#if ADD_USER_LOGIN_INFO
+			EOS_Connect_UserLoginInfo UserLoginInfo = {};
+			UserLoginInfo.ApiVersion = EOS_CONNECT_USERLOGININFO_API_LATEST;
+			const FTCHARToUTF8 DisplayNameUtf8(*GetPlatformDisplayName(LocalUserNum));
+			UserLoginInfo.DisplayName = DisplayNameUtf8.Get();
+
+			Options.UserLoginInfo = &UserLoginInfo;
+#endif
+
 			FConnectLoginCallback* CallbackObj = new FConnectLoginCallback(FUserManagerEOSWeakPtr(AsShared()));
 			CallbackObj->CallbackLambda = [this, LocalUserNum](const EOS_Connect_LoginCallbackInfo* Data)
 			{
diff --git a/Engine/Plugins/Online/OnlineSubsystemEOS/Source/OnlineSubsystemEOS/Private/UserManagerEOS.h b/Engine/Plugins/Online/OnlineSubsystemEOS/Source/OnlineSubsystemEOS/Private/UserManagerEOS.h
--- a/Engine/Plugins/Online/OnlineSubsystemEOS/Source/OnlineSubsystemEOS/Private/UserManagerEOS.h
+++ b/Engine/Plugins/Online/OnlineSubsystemEOS/Source/OnlineSubsystemEOS/Private/UserManagerEOS.h
@@ -401,8 +401,9 @@ private:
 	void UpdatePresence(EOS_EpicAccountId AccountId);
 	void UpdateFriendPresence(const FString& FriendId, FOnlineUserPresenceRef Presence);
 
-	IOnlineSubsystem* GetPlatformOSS();
-	void GetPlatformAuthToken(int32 LocalUserNum, const FOnGetLinkedAccountAuthTokenCompleteDelegate& Delegate);
+	IOnlineSubsystem* GetPlatformOSS() const;
+	void GetPlatformAuthToken(int32 LocalUserNum, const FOnGetLinkedAccountAuthTokenCompleteDelegate& Delegate) const;
+	FString GetPlatformDisplayName(int32 LocalUserNum) const;
 
 	/** Cached pointer to owning subsystem */
 	FOnlineSubsystemEOS* EOSSubsystem;
