Re-adding EOS Lobby attributes used as search parameters for consistency with 5.0

diff --git a/Engine/Plugins/Online/OnlineSubsystemEOS/Source/OnlineSubsystemEOS/Private/OnlineSessionEOS.cpp b/Engine/Plugins/Online/OnlineSubsystemEOS/Source/OnlineSubsystemEOS/Private/OnlineSessionEOS.cpp
--- a/Engine/Plugins/Online/OnlineSubsystemEOS/Source/OnlineSubsystemEOS/Private/OnlineSessionEOS.cpp
+++ b/Engine/Plugins/Online/OnlineSubsystemEOS/Source/OnlineSubsystemEOS/Private/OnlineSessionEOS.cpp
@@ -1041,6 +1041,11 @@ void FOnlineSessionEOS::AddAttribute(EOS_HSessionModification SessionModHandle,
 
 void FOnlineSessionEOS::SetAttributes(EOS_HSessionModification SessionModHandle, FNamedOnlineSession* Session)
 {
+	// The first will let us find it on session searches
+	const FString SearchPresence(SEARCH_PRESENCE.ToString());
+	const FAttributeOptions SearchPresenceAttribute(TCHAR_TO_UTF8(*SearchPresence), true);
+	AddAttribute(SessionModHandle, &SearchPresenceAttribute);
+
 	FAttributeOptions Opt1("NumPrivateConnections", Session->SessionSettings.NumPrivateConnections);
 	AddAttribute(SessionModHandle, &Opt1);
 
@@ -3554,6 +3559,16 @@ void FOnlineSessionEOS::SetLobbyAttributes(EOS_HLobbyModification LobbyModificat
 {
 	check(Session != nullptr);
 
+	// The first will let us find it on session searches
+	const FString SearchPresence(SEARCH_PRESENCE.ToString());
+	const FLobbyAttributeOptions SearchPresenceAttribute(TCHAR_TO_UTF8(*SearchPresence), true);
+	AddLobbyAttribute(LobbyModificationHandle, &SearchPresenceAttribute);
+
+	// The second will let us find it on lobby searches
+	const FString SearchLobbies(SEARCH_LOBBIES.ToString());
+	const FLobbyAttributeOptions SearchLobbiesAttribute(TCHAR_TO_UTF8(*SearchLobbies), true);
+	AddLobbyAttribute(LobbyModificationHandle, &SearchLobbiesAttribute);
+
 	// We set the session's owner id and name
 	const FLobbyAttributeOptions OwnerId("OwningUserId", TCHAR_TO_UTF8(*Session->OwningUserId->ToString()));
 	AddLobbyAttribute(LobbyModificationHandle, &OwnerId);
