Fix server crash on shutdown
#rnx

diff --git a/Engine/Plugins/Online/OnlineSubsystemEOS/Source/OnlineSubsystemEOS/Private/UserManagerEOS.cpp b/Engine/Plugins/Online/OnlineSubsystemEOS/Source/OnlineSubsystemEOS/Private/UserManagerEOS.cpp
--- a/Engine/Plugins/Online/OnlineSubsystemEOS/Source/OnlineSubsystemEOS/Private/UserManagerEOS.cpp
+++ b/Engine/Plugins/Online/OnlineSubsystemEOS/Source/OnlineSubsystemEOS/Private/UserManagerEOS.cpp
@@ -141,6 +141,8 @@ FUserManagerEOS::FUserManagerEOS(FOnlineSubsystemEOS* InSubsystem)
 	, FriendsNotificationCallback(nullptr)
 	, PresenceNotificationId(0)
 	, PresenceNotificationCallback(nullptr)
+	, DisplaySettingsUpdatedId(0)
+	, DisplaySettingsUpdatedCallback(nullptr)
 {
 }
 
@@ -172,6 +174,8 @@ void FUserManagerEOS::Shutdown()
 		// Removing subscription to external ui display change event
 		EOS_UI_RemoveNotifyDisplaySettingsUpdated(EOSSubsystem->UIHandle, DisplaySettingsUpdatedId);
 
+	if (DisplaySettingsUpdatedCallback)
+	{
 		delete DisplaySettingsUpdatedCallback;
 	}
 }
