Fix bad auto-merge
#rnx

diff --git a/Engine/Plugins/Online/OnlineSubsystemEOS/Source/OnlineSubsystemEOS/Private/UserManagerEOS.cpp b/Engine/Plugins/Online/OnlineSubsystemEOS/Source/OnlineSubsystemEOS/Private/UserManagerEOS.cpp
--- a/Engine/Plugins/Online/OnlineSubsystemEOS/Source/OnlineSubsystemEOS/Private/UserManagerEOS.cpp
+++ b/Engine/Plugins/Online/OnlineSubsystemEOS/Source/OnlineSubsystemEOS/Private/UserManagerEOS.cpp
@@ -174,9 +174,10 @@ void FUserManagerEOS::Shutdown()
 		// Removing subscription to external ui display change event
 		EOS_UI_RemoveNotifyDisplaySettingsUpdated(EOSSubsystem->UIHandle, DisplaySettingsUpdatedId);
 
-	if (DisplaySettingsUpdatedCallback)
-	{
-		delete DisplaySettingsUpdatedCallback;
+		if (DisplaySettingsUpdatedCallback)
+		{
+			delete DisplaySettingsUpdatedCallback;
+		}
 	}
 }
 
