- Move OSSEOS headers to Private folder
- Introduce new IOnlineSubsystemEOS.h with the public interface.
- Add missing includes to OnlineSubsystemEOSPrivate.h, so the ONLINE_LOG_PREFIX is visible to UE_LOG_ONLINE

diff --git a/Engine/Plugins/Online/OnlineSubsystemEOS/Source/OnlineSubsystemEOS/Private/EOSHelpers.cpp b/Engine/Plugins/Online/OnlineSubsystemEOS/Source/OnlineSubsystemEOS/Private/EOSHelpers.cpp
--- a/Engine/Plugins/Online/OnlineSubsystemEOS/Source/OnlineSubsystemEOS/Private/EOSHelpers.cpp
+++ b/Engine/Plugins/Online/OnlineSubsystemEOS/Source/OnlineSubsystemEOS/Private/EOSHelpers.cpp
@@ -7,6 +7,7 @@
 #include "IEOSSDKManager.h"
 #include "OnlineError.h"
 #include "OnlineSubsystemEOS.h"
+#include "OnlineSubsystemEOSPrivate.h"
 #include "UserManagerEOS.h"
 
 void FEOSHelpers::PlatformAuthCredentials(EOS_Auth_Credentials &Credentials)
diff --git a/Engine/Plugins/Online/OnlineSubsystemEOS/Source/OnlineSubsystemEOS/Private/EOSSettings.cpp b/Engine/Plugins/Online/OnlineSubsystemEOS/Source/OnlineSubsystemEOS/Private/EOSSettings.cpp
--- a/Engine/Plugins/Online/OnlineSubsystemEOS/Source/OnlineSubsystemEOS/Private/EOSSettings.cpp
+++ b/Engine/Plugins/Online/OnlineSubsystemEOS/Source/OnlineSubsystemEOS/Private/EOSSettings.cpp
@@ -3,6 +3,7 @@
 #include "EOSSettings.h"
 #include "OnlineSubsystemEOS.h"
 #include "OnlineSubsystemEOSModule.h"
+#include "OnlineSubsystemEOSPrivate.h"
 
 #include "Algo/Transform.h"
 #include "Misc/CommandLine.h"
diff --git a/Engine/Plugins/Online/OnlineSubsystemEOS/Source/OnlineSubsystemEOS/Private/OnlineAchievementsEOS.cpp b/Engine/Plugins/Online/OnlineSubsystemEOS/Source/OnlineSubsystemEOS/Private/OnlineAchievementsEOS.cpp
--- a/Engine/Plugins/Online/OnlineSubsystemEOS/Source/OnlineSubsystemEOS/Private/OnlineAchievementsEOS.cpp
+++ b/Engine/Plugins/Online/OnlineSubsystemEOS/Source/OnlineSubsystemEOS/Private/OnlineAchievementsEOS.cpp
@@ -3,11 +3,11 @@
 #include "OnlineAchievementsEOS.h"
 #include "OnlineSubsystem.h"
 #include "OnlineSubsystemEOS.h"
+#include "OnlineSubsystemEOSPrivate.h"
 #include "OnlineSubsystemEOSTypes.h"
 #include "OnlineStatsEOS.h"
 #include "UserManagerEOS.h"
 
-
 #if WITH_EOS_SDK
 #include "eos_achievements.h"
 
diff --git a/Engine/Plugins/Online/OnlineSubsystemEOS/Source/OnlineSubsystemEOS/Private/OnlineAchievementsEOS.h b/Engine/Plugins/Online/OnlineSubsystemEOS/Source/OnlineSubsystemEOS/Private/OnlineAchievementsEOS.h
--- a/Engine/Plugins/Online/OnlineSubsystemEOS/Source/OnlineSubsystemEOS/Private/OnlineAchievementsEOS.h
+++ b/Engine/Plugins/Online/OnlineSubsystemEOS/Source/OnlineSubsystemEOS/Private/OnlineAchievementsEOS.h
@@ -5,7 +5,6 @@
 #include "CoreMinimal.h"
 #include "Online/CoreOnline.h"
 #include "Interfaces/OnlineAchievementsInterface.h"
-#include "OnlineSubsystemEOSPackage.h"
 #include "OnlineSubsystemEOSTypes.h"
 
 class FOnlineSubsystemEOS;
@@ -36,7 +35,6 @@ public:
 #endif
 // ~IOnlineAchievements Interface
 
-PACKAGE_SCOPE:
 	FOnlineAchievementsEOS(FOnlineSubsystemEOS* InSubsystem)
 		: EOSSubsystem(InSubsystem)
 	{
diff --git a/Engine/Plugins/Online/OnlineSubsystemEOS/Source/OnlineSubsystemEOS/Private/OnlineLeaderboardsEOS.cpp b/Engine/Plugins/Online/OnlineSubsystemEOS/Source/OnlineSubsystemEOS/Private/OnlineLeaderboardsEOS.cpp
--- a/Engine/Plugins/Online/OnlineSubsystemEOS/Source/OnlineSubsystemEOS/Private/OnlineLeaderboardsEOS.cpp
+++ b/Engine/Plugins/Online/OnlineSubsystemEOS/Source/OnlineSubsystemEOS/Private/OnlineLeaderboardsEOS.cpp
@@ -3,6 +3,7 @@
 #include "OnlineLeaderboardsEOS.h"
 #include "OnlineSubsystem.h"
 #include "OnlineSubsystemEOS.h"
+#include "OnlineSubsystemEOSPrivate.h"
 #include "OnlineSubsystemEOSTypes.h"
 #include "UserManagerEOS.h"
 #include "OnlineStatsEOS.h"
diff --git a/Engine/Plugins/Online/OnlineSubsystemEOS/Source/OnlineSubsystemEOS/Private/OnlineLeaderboardsEOS.h b/Engine/Plugins/Online/OnlineSubsystemEOS/Source/OnlineSubsystemEOS/Private/OnlineLeaderboardsEOS.h
--- a/Engine/Plugins/Online/OnlineSubsystemEOS/Source/OnlineSubsystemEOS/Private/OnlineLeaderboardsEOS.h
+++ b/Engine/Plugins/Online/OnlineSubsystemEOS/Source/OnlineSubsystemEOS/Private/OnlineLeaderboardsEOS.h
@@ -3,7 +3,6 @@
 #pragma once
 
 #include "Interfaces/OnlineLeaderboardInterface.h"
-#include "OnlineSubsystemEOSPackage.h"
 #include "OnlineSubsystemEOSTypes.h"
 
 class FOnlineSubsystemEOS;
@@ -35,7 +34,6 @@ public:
 	virtual bool WriteOnlinePlayerRatings(const FName& SessionName, int32 LeaderboardId, const TArray<FOnlinePlayerScore>& PlayerScores) override;
 // ~IOnlineLeaderboards Interface
 
-PACKAGE_SCOPE:
 	FOnlineLeaderboardsEOS(FOnlineSubsystemEOS* InSubsystem)
 		: EOSSubsystem(InSubsystem)
 	{
diff --git a/Engine/Plugins/Online/OnlineSubsystemEOS/Source/OnlineSubsystemEOS/Private/OnlineSessionEOS.cpp b/Engine/Plugins/Online/OnlineSubsystemEOS/Source/OnlineSubsystemEOS/Private/OnlineSessionEOS.cpp
--- a/Engine/Plugins/Online/OnlineSubsystemEOS/Source/OnlineSubsystemEOS/Private/OnlineSessionEOS.cpp
+++ b/Engine/Plugins/Online/OnlineSubsystemEOS/Source/OnlineSubsystemEOS/Private/OnlineSessionEOS.cpp
@@ -4,6 +4,7 @@
 #include "Misc/Guid.h"
 #include "OnlineSubsystem.h"
 #include "OnlineSubsystemEOS.h"
+#include "OnlineSubsystemEOSPrivate.h"
 #include "OnlineSubsystemEOSTypes.h"
 #include "UserManagerEOS.h"
 #include "OnlineSubsystemUtils.h"
diff --git a/Engine/Plugins/Online/OnlineSubsystemEOS/Source/OnlineSubsystemEOS/Private/OnlineSessionEOS.h b/Engine/Plugins/Online/OnlineSubsystemEOS/Source/OnlineSubsystemEOS/Private/OnlineSessionEOS.h
--- a/Engine/Plugins/Online/OnlineSubsystemEOS/Source/OnlineSubsystemEOS/Private/OnlineSessionEOS.h
+++ b/Engine/Plugins/Online/OnlineSubsystemEOS/Source/OnlineSubsystemEOS/Private/OnlineSessionEOS.h
@@ -5,7 +5,6 @@
 #include "Misc/ScopeLock.h"
 #include "OnlineSessionSettings.h"
 #include "Interfaces/OnlineSessionInterface.h"
-#include "OnlineSubsystemEOSPackage.h"
 #include "LANBeacon.h"
 #include "OnlineSubsystemEOSTypes.h"
 
@@ -137,7 +136,6 @@ public:
 	virtual void DumpSessionState() override;
 // ~IOnlineSession Interface
 
-PACKAGE_SCOPE:
 	/** Critical sections for thread safe operation of session lists */
 	mutable FCriticalSection SessionLock;
 
diff --git a/Engine/Plugins/Online/OnlineSubsystemEOS/Source/OnlineSubsystemEOS/Private/OnlineStatsEOS.cpp b/Engine/Plugins/Online/OnlineSubsystemEOS/Source/OnlineSubsystemEOS/Private/OnlineStatsEOS.cpp
--- a/Engine/Plugins/Online/OnlineSubsystemEOS/Source/OnlineSubsystemEOS/Private/OnlineStatsEOS.cpp
+++ b/Engine/Plugins/Online/OnlineSubsystemEOS/Source/OnlineSubsystemEOS/Private/OnlineStatsEOS.cpp
@@ -3,6 +3,7 @@
 #include "OnlineStatsEOS.h"
 #include "OnlineSubsystem.h"
 #include "OnlineSubsystemEOS.h"
+#include "OnlineSubsystemEOSPrivate.h"
 #include "OnlineSubsystemEOSTypes.h"
 #include "UserManagerEOS.h"
 
diff --git a/Engine/Plugins/Online/OnlineSubsystemEOS/Source/OnlineSubsystemEOS/Private/OnlineStatsEOS.h b/Engine/Plugins/Online/OnlineSubsystemEOS/Source/OnlineSubsystemEOS/Private/OnlineStatsEOS.h
--- a/Engine/Plugins/Online/OnlineSubsystemEOS/Source/OnlineSubsystemEOS/Private/OnlineStatsEOS.h
+++ b/Engine/Plugins/Online/OnlineSubsystemEOS/Source/OnlineSubsystemEOS/Private/OnlineStatsEOS.h
@@ -3,7 +3,6 @@
 #pragma once
 
 #include "Interfaces/OnlineStatsInterface.h"
-#include "OnlineSubsystemEOSPackage.h"
 #include "OnlineSubsystemEOSTypes.h"
 
 class FOnlineSubsystemEOS;
@@ -32,7 +31,6 @@ public:
 #endif
 // ~IOnlineStats Interface
 
-PACKAGE_SCOPE:
 	FOnlineStatsEOS(FOnlineSubsystemEOS* InSubsystem)
 		: EOSSubsystem(InSubsystem)
 	{
diff --git a/Engine/Plugins/Online/OnlineSubsystemEOS/Source/OnlineSubsystemEOS/Private/OnlineStoreEOS.cpp b/Engine/Plugins/Online/OnlineSubsystemEOS/Source/OnlineSubsystemEOS/Private/OnlineStoreEOS.cpp
--- a/Engine/Plugins/Online/OnlineSubsystemEOS/Source/OnlineSubsystemEOS/Private/OnlineStoreEOS.cpp
+++ b/Engine/Plugins/Online/OnlineSubsystemEOS/Source/OnlineSubsystemEOS/Private/OnlineStoreEOS.cpp
@@ -5,6 +5,7 @@
 #if WITH_EOS_SDK
 
 #include "OnlineSubsystemEOS.h"
+#include "OnlineSubsystemEOSPrivate.h"
 #include "UserManagerEOS.h"
 #include "eos_ecom.h"
 
diff --git a/Engine/Plugins/Online/OnlineSubsystemEOS/Source/OnlineSubsystemEOS/Private/OnlineStoreEOS.h b/Engine/Plugins/Online/OnlineSubsystemEOS/Source/OnlineSubsystemEOS/Private/OnlineStoreEOS.h
--- a/Engine/Plugins/Online/OnlineSubsystemEOS/Source/OnlineSubsystemEOS/Private/OnlineStoreEOS.h
+++ b/Engine/Plugins/Online/OnlineSubsystemEOS/Source/OnlineSubsystemEOS/Private/OnlineStoreEOS.h
@@ -6,7 +6,6 @@
 #include "Online/CoreOnline.h"
 #include "Interfaces/OnlineStoreInterfaceV2.h"
 #include "Interfaces/OnlinePurchaseInterface.h"
-#include "OnlineSubsystemEOSPackage.h"
 #include "OnlineSubsystemEOSTypes.h"
 
 #if WITH_EOS_SDK
@@ -44,7 +43,6 @@ public:
 	virtual void FinalizeReceiptValidationInfo(const FUniqueNetId& UserId, FString& InReceiptValidationInfo, const FOnFinalizeReceiptValidationInfoComplete& Delegate) override;
 // End IOnlinePurchase
 
-PACKAGE_SCOPE:
 	FOnlineStoreEOS(FOnlineSubsystemEOS* InSubsystem);
 
 	bool HandleEcomExec(UWorld* InWorld, const TCHAR* Cmd, FOutputDevice& Ar);
diff --git a/Engine/Plugins/Online/OnlineSubsystemEOS/Source/OnlineSubsystemEOS/Private/OnlineSubsystemEOS.cpp b/Engine/Plugins/Online/OnlineSubsystemEOS/Source/OnlineSubsystemEOS/Private/OnlineSubsystemEOS.cpp
--- a/Engine/Plugins/Online/OnlineSubsystemEOS/Source/OnlineSubsystemEOS/Private/OnlineSubsystemEOS.cpp
+++ b/Engine/Plugins/Online/OnlineSubsystemEOS/Source/OnlineSubsystemEOS/Private/OnlineSubsystemEOS.cpp
@@ -1,6 +1,7 @@
 // Copyright Epic Games, Inc. All Rights Reserved.
 
 #include "OnlineSubsystemEOS.h"
+#include "OnlineSubsystemEOSPrivate.h"
 #include "OnlineSubsystemUtils.h"
 #include "UserManagerEOS.h"
 #include "OnlineSessionEOS.h"
@@ -541,7 +542,7 @@ FText FOnlineSubsystemEOS::GetOnlineServiceName() const
 }
 
 FOnlineSubsystemEOS::FOnlineSubsystemEOS(FName InInstanceName) :
-	FOnlineSubsystemImpl(EOS_SUBSYSTEM, InInstanceName)
+	IOnlineSubsystemEOS(EOS_SUBSYSTEM, InInstanceName)
 	, EOSSDKManager(nullptr)
 	, AuthHandle(nullptr)
 	, UIHandle(nullptr)
diff --git a/Engine/Plugins/Online/OnlineSubsystemEOS/Source/OnlineSubsystemEOS/Private/OnlineSubsystemEOS.h b/Engine/Plugins/Online/OnlineSubsystemEOS/Source/OnlineSubsystemEOS/Private/OnlineSubsystemEOS.h
new file mode 100644
--- /dev/null
+++ b/Engine/Plugins/Online/OnlineSubsystemEOS/Source/OnlineSubsystemEOS/Private/OnlineSubsystemEOS.h
@@ -0,0 +1,217 @@
+// Copyright Epic Games, Inc. All Rights Reserved.
+
+#pragma once
+
+#include "CoreMinimal.h"
+#include "IOnlineSubsystemEOS.h"
+#include "OnlineSubsystemNames.h"
+
+#include COMPILED_PLATFORM_HEADER(EOSHelpers.h)
+
+DECLARE_STATS_GROUP(TEXT("EOS"), STATGROUP_EOS, STATCAT_Advanced);
+
+#if WITH_EOS_SDK
+
+#include "eos_sdk.h"
+
+class FSocketSubsystemEOS;
+class IEOSSDKManager;
+using IEOSPlatformHandlePtr = TSharedPtr<class IEOSPlatformHandle, ESPMode::ThreadSafe>;
+
+class IVoiceChatUser;
+class FEOSVoiceChatUser;
+using IVoiceChatPtr = TSharedPtr<class IVoiceChat, ESPMode::ThreadSafe>;
+using FOnlineSubsystemEOSVoiceChatUserWrapperRef = TSharedRef<class FOnlineSubsystemEOSVoiceChatUserWrapper, ESPMode::ThreadSafe>;
+
+class FUserManagerEOS;
+typedef TSharedPtr<class FUserManagerEOS, ESPMode::ThreadSafe> FUserManagerEOSPtr;
+
+class FOnlineSessionEOS;
+typedef TSharedPtr<class FOnlineSessionEOS, ESPMode::ThreadSafe> FOnlineSessionEOSPtr;
+
+class FOnlineStatsEOS;
+typedef TSharedPtr<class FOnlineStatsEOS, ESPMode::ThreadSafe> FOnlineStatsEOSPtr;
+
+class FOnlineLeaderboardsEOS;
+typedef TSharedPtr<class FOnlineLeaderboardsEOS, ESPMode::ThreadSafe> FOnlineLeaderboardsEOSPtr;
+
+class FOnlineAchievementsEOS;
+typedef TSharedPtr<class FOnlineAchievementsEOS, ESPMode::ThreadSafe> FOnlineAchievementsEOSPtr;
+
+class FOnlineStoreEOS;
+typedef TSharedPtr<class FOnlineStoreEOS, ESPMode::ThreadSafe> FOnlineStoreEOSPtr;
+
+class FOnlineTitleFileEOS;
+typedef TSharedPtr<class FOnlineTitleFileEOS, ESPMode::ThreadSafe> FOnlineTitleFileEOSPtr;
+
+class FOnlineUserCloudEOS;
+typedef TSharedPtr<class FOnlineUserCloudEOS, ESPMode::ThreadSafe> FOnlineUserCloudEOSPtr;
+
+typedef TSharedPtr<FPlatformEOSHelpers, ESPMode::ThreadSafe> FPlatformEOSHelpersPtr;
+
+/**
+ *	OnlineSubsystemEOS - Implementation of the online subsystem for EOS services
+ */
+class ONLINESUBSYSTEMEOS_API FOnlineSubsystemEOS : 
+	public IOnlineSubsystemEOS
+{
+public:
+	virtual ~FOnlineSubsystemEOS() = default;
+
+	/** Used to be called before RHIInit() */
+	static void ModuleInit();
+	static void ModuleShutdown();
+
+	FPlatformEOSHelpersPtr GetEOSHelpers() { return EOSHelpersPtr; };
+
+// IOnlineSubsystemEOS
+	virtual IVoiceChatUser* GetVoiceChatUserInterface(const FUniqueNetId& LocalUserId) override;
+	virtual IEOSPlatformHandlePtr GetEOSPlatformHandle() const override { return EOSPlatformHandle; };
+
+// IOnlineSubsystem
+	virtual IOnlineSessionPtr GetSessionInterface() const override;
+	virtual IOnlineFriendsPtr GetFriendsInterface() const override;
+	virtual IOnlineSharedCloudPtr GetSharedCloudInterface() const override;
+	virtual IOnlineUserCloudPtr GetUserCloudInterface() const override;
+	virtual IOnlineEntitlementsPtr GetEntitlementsInterface() const override;
+	virtual IOnlineLeaderboardsPtr GetLeaderboardsInterface() const override;
+	virtual IOnlineVoicePtr GetVoiceInterface() const override;
+	virtual IOnlineExternalUIPtr GetExternalUIInterface() const override;	
+	virtual IOnlineIdentityPtr GetIdentityInterface() const override;
+	virtual IOnlineTitleFilePtr GetTitleFileInterface() const override;
+	virtual IOnlineStoreV2Ptr GetStoreV2Interface() const override;
+	virtual IOnlinePurchasePtr GetPurchaseInterface() const override;
+	virtual IOnlineAchievementsPtr GetAchievementsInterface() const override;
+	virtual IOnlineUserPtr GetUserInterface() const override;
+	virtual IOnlinePresencePtr GetPresenceInterface() const override;
+	virtual FText GetOnlineServiceName() const override;
+	virtual IOnlineStatsPtr GetStatsInterface() const override;
+	virtual bool Exec(class UWorld* InWorld, const TCHAR* Cmd, FOutputDevice& Ar) override;
+	virtual void ReloadConfigs(const TSet<FString>& ConfigSections) override;
+
+	virtual IOnlineGroupsPtr GetGroupsInterface() const override { return nullptr; }
+	virtual IOnlinePartyPtr GetPartyInterface() const override { return nullptr; }
+	virtual IOnlineTimePtr GetTimeInterface() const override { return nullptr; }
+	virtual IOnlineEventsPtr GetEventsInterface() const override { return nullptr; }
+	virtual IOnlineSharingPtr GetSharingInterface() const override { return nullptr; }
+	virtual IOnlineMessagePtr GetMessageInterface() const override { return nullptr; }
+	virtual IOnlineChatPtr GetChatInterface() const override { return nullptr; }
+	virtual IOnlineTurnBasedPtr GetTurnBasedInterface() const override { return nullptr; }
+	virtual IOnlineTournamentPtr GetTournamentInterface() const override { return nullptr; }
+//~IOnlineSubsystem
+
+	virtual bool Init() override;
+	virtual bool Shutdown() override;
+	virtual FString GetAppId() const override;
+
+// FTSTickerObjectBase
+	virtual bool Tick(float DeltaTime) override;
+
+	/** Only the factory makes instances */
+	FOnlineSubsystemEOS() = delete;
+	explicit FOnlineSubsystemEOS(FName InInstanceName);
+
+	FString ProductId;
+
+	IEOSSDKManager* EOSSDKManager;
+
+	/** EOS handles */
+	IEOSPlatformHandlePtr EOSPlatformHandle;
+	EOS_HAuth AuthHandle;
+	EOS_HUI UIHandle;
+	EOS_HFriends FriendsHandle;
+	EOS_HUserInfo UserInfoHandle;
+	EOS_HPresence PresenceHandle;
+	EOS_HConnect ConnectHandle;
+	EOS_HSessions SessionsHandle;
+	EOS_HStats StatsHandle;
+	EOS_HLeaderboards LeaderboardsHandle;
+	EOS_HMetrics MetricsHandle;
+	EOS_HAchievements AchievementsHandle;
+	EOS_HEcom EcomHandle;
+	EOS_HTitleStorage TitleStorageHandle;
+	EOS_HPlayerDataStorage PlayerDataStorageHandle;
+
+	/** Manager that handles all user interfaces */
+	FUserManagerEOSPtr UserManager;
+	/** The session interface object */
+	FOnlineSessionEOSPtr SessionInterfacePtr;
+	/** Stats interface pointer */
+	FOnlineStatsEOSPtr StatsInterfacePtr;
+	/** Leaderboards interface pointer */
+	FOnlineLeaderboardsEOSPtr LeaderboardsInterfacePtr;
+	FOnlineAchievementsEOSPtr AchievementsInterfacePtr;
+	/** EGS store interface pointer */
+	FOnlineStoreEOSPtr StoreInterfacePtr;
+	/** Title File interface pointer */
+	FOnlineTitleFileEOSPtr TitleFileInterfacePtr;
+	/** User Cloud interface pointer */
+	FOnlineUserCloudEOSPtr UserCloudInterfacePtr;
+
+	bool bWasLaunchedByEGS;
+	bool bIsDefaultOSS;
+	bool bIsPlatformOSS;
+
+	TSharedPtr<FSocketSubsystemEOS, ESPMode::ThreadSafe> SocketSubsystem;
+
+	static FPlatformEOSHelpersPtr EOSHelpersPtr;
+
+	FEOSVoiceChatUser* GetEOSVoiceChatUserInterface(const FUniqueNetId& LocalUserId);
+	void ReleaseVoiceChatUserInterface(const FUniqueNetId& LocalUserId);
+
+private:
+	bool PlatformCreate();
+
+	IVoiceChatPtr VoiceChatInterface;
+	TUniqueNetIdMap<FOnlineSubsystemEOSVoiceChatUserWrapperRef> LocalVoiceChatUsers;
+};
+
+#else
+
+class ONLINESUBSYSTEMEOS_API FOnlineSubsystemEOS :
+	public FOnlineSubsystemImpl
+{
+public:
+	explicit FOnlineSubsystemEOS(FName InInstanceName) :
+		FOnlineSubsystemImpl(EOS_SUBSYSTEM, InInstanceName)
+	{
+	}
+
+	virtual ~FOnlineSubsystemEOS() = default;
+
+	virtual IOnlineSessionPtr GetSessionInterface() const override { return nullptr; }
+	virtual IOnlineFriendsPtr GetFriendsInterface() const override { return nullptr; }
+	virtual IOnlineGroupsPtr GetGroupsInterface() const override { return nullptr; }
+	virtual IOnlinePartyPtr GetPartyInterface() const override { return nullptr; }
+	virtual IOnlineSharedCloudPtr GetSharedCloudInterface() const override { return nullptr; }
+	virtual IOnlineUserCloudPtr GetUserCloudInterface() const override { return nullptr; }
+	virtual IOnlineEntitlementsPtr GetEntitlementsInterface() const override { return nullptr; }
+	virtual IOnlineLeaderboardsPtr GetLeaderboardsInterface() const override { return nullptr; }
+	virtual IOnlineVoicePtr GetVoiceInterface() const override { return nullptr; }
+	virtual IOnlineExternalUIPtr GetExternalUIInterface() const override { return nullptr; }
+	virtual IOnlineTimePtr GetTimeInterface() const override { return nullptr; }
+	virtual IOnlineIdentityPtr GetIdentityInterface() const override { return nullptr; }
+	virtual IOnlineTitleFilePtr GetTitleFileInterface() const override { return nullptr; }
+	virtual IOnlineStoreV2Ptr GetStoreV2Interface() const override { return nullptr; }
+	virtual IOnlinePurchasePtr GetPurchaseInterface() const override { return nullptr; }
+	virtual IOnlineEventsPtr GetEventsInterface() const override { return nullptr; }
+	virtual IOnlineAchievementsPtr GetAchievementsInterface() const override { return nullptr; }
+	virtual IOnlineSharingPtr GetSharingInterface() const override { return nullptr; }
+	virtual IOnlineUserPtr GetUserInterface() const override { return nullptr; }
+	virtual IOnlineMessagePtr GetMessageInterface() const override { return nullptr; }
+	virtual IOnlinePresencePtr GetPresenceInterface() const override { return nullptr; }
+	virtual IOnlineChatPtr GetChatInterface() const override { return nullptr; }
+	virtual IOnlineStatsPtr GetStatsInterface() const override { return nullptr; }
+	virtual IOnlineTurnBasedPtr GetTurnBasedInterface() const override { return nullptr; }
+	virtual IOnlineTournamentPtr GetTournamentInterface() const override { return nullptr; }
+	virtual FText GetOnlineServiceName() const override { return NSLOCTEXT("OnlineSubsystemEOS", "OnlineServiceName", "EOS"); }
+
+	virtual bool Init() override { return false; }
+	virtual bool Shutdown() override { return true; }
+	virtual FString GetAppId() const override { return TEXT(""); }
+};
+
+#endif
+
+typedef TSharedPtr<FOnlineSubsystemEOS, ESPMode::ThreadSafe> FOnlineSubsystemEOSPtr;
+
diff --git a/Engine/Plugins/Online/OnlineSubsystemEOS/Source/OnlineSubsystemEOS/Private/OnlineSubsystemEOSModule.h b/Engine/Plugins/Online/OnlineSubsystemEOS/Source/OnlineSubsystemEOS/Private/OnlineSubsystemEOSModule.h
new file mode 100644
--- /dev/null
+++ b/Engine/Plugins/Online/OnlineSubsystemEOS/Source/OnlineSubsystemEOS/Private/OnlineSubsystemEOSModule.h
@@ -0,0 +1,45 @@
+// Copyright Epic Games, Inc. All Rights Reserved.
+
+#pragma once
+
+#include "CoreMinimal.h"
+#include "Modules/ModuleInterface.h"
+
+/**
+ * Online subsystem module class  (EOS Implementation)
+ * Code related to the loading of the EOS module
+ */
+class FOnlineSubsystemEOSModule : public IModuleInterface
+{
+private:
+
+	/** Class responsible for creating instance(s) of the subsystem */
+	class FOnlineFactoryEOS* EOSFactory;
+
+public:
+
+	FOnlineSubsystemEOSModule() :
+		EOSFactory(NULL)
+	{}
+
+	virtual ~FOnlineSubsystemEOSModule() {}
+
+#if WITH_EDITOR
+	void OnPostEngineInit();
+	void OnPreExit();
+#endif
+
+	// IModuleInterface
+
+	virtual void StartupModule() override;
+	virtual void ShutdownModule() override;
+	virtual bool SupportsDynamicReloading() override
+	{
+		return false;
+	}
+
+	virtual bool SupportsAutomaticShutdown() override
+	{
+		return false;
+	}
+};
diff --git a/Engine/Plugins/Online/OnlineSubsystemEOS/Source/OnlineSubsystemEOS/Private/OnlineSubsystemEOSTypes.h b/Engine/Plugins/Online/OnlineSubsystemEOS/Source/OnlineSubsystemEOS/Private/OnlineSubsystemEOSTypes.h
--- a/Engine/Plugins/Online/OnlineSubsystemEOS/Source/OnlineSubsystemEOS/Private/OnlineSubsystemEOSTypes.h
+++ b/Engine/Plugins/Online/OnlineSubsystemEOS/Source/OnlineSubsystemEOS/Private/OnlineSubsystemEOSTypes.h
@@ -12,8 +12,6 @@
 #include "EOSShared.h"
 #include "IPAddress.h"
 
-#include "OnlineSubsystemEOSPackage.h"
-
 #define EOS_OSS_STRING_BUFFER_LENGTH 256 + 1 // 256 plus null terminator
 
 class FOnlineSubsystemEOS;
@@ -80,7 +78,6 @@ public:
 		return EOS_ID_BYTE_SIZE;
 	}
 
-PACKAGE_SCOPE:
 	void UpdateNetIdStr(const FString& InNetIdStr)
 	{
 		UniqueNetIdStr = InNetIdStr;
diff --git a/Engine/Plugins/Online/OnlineSubsystemEOS/Source/OnlineSubsystemEOS/Private/OnlineSubsystemModuleEOS.cpp b/Engine/Plugins/Online/OnlineSubsystemEOS/Source/OnlineSubsystemEOS/Private/OnlineSubsystemModuleEOS.cpp
--- a/Engine/Plugins/Online/OnlineSubsystemEOS/Source/OnlineSubsystemEOS/Private/OnlineSubsystemModuleEOS.cpp
+++ b/Engine/Plugins/Online/OnlineSubsystemEOS/Source/OnlineSubsystemEOS/Private/OnlineSubsystemModuleEOS.cpp
@@ -1,6 +1,7 @@
 // Copyright Epic Games, Inc. All Rights Reserved.
 
 #include "OnlineSubsystemEOSModule.h"
+#include "OnlineSubsystemEOSPrivate.h"
 #include "OnlineSubsystemModule.h"
 #include "OnlineSubsystemNames.h"
 #include "OnlineSubsystemEOS.h"
diff --git a/Engine/Plugins/Online/OnlineSubsystemEOS/Source/OnlineSubsystemEOS/Private/OnlineTitleFileEOS.cpp b/Engine/Plugins/Online/OnlineSubsystemEOS/Source/OnlineSubsystemEOS/Private/OnlineTitleFileEOS.cpp
--- a/Engine/Plugins/Online/OnlineSubsystemEOS/Source/OnlineSubsystemEOS/Private/OnlineTitleFileEOS.cpp
+++ b/Engine/Plugins/Online/OnlineSubsystemEOS/Source/OnlineSubsystemEOS/Private/OnlineTitleFileEOS.cpp
@@ -3,6 +3,7 @@
 #include "OnlineTitleFileEOS.h"
 #include "OnlineSubsystem.h"
 #include "OnlineSubsystemEOS.h"
+#include "OnlineSubsystemEOSPrivate.h"
 #include "OnlineSubsystemEOSTypes.h"
 #include "UserManagerEOS.h"
 #include "EOSSettings.h"
diff --git a/Engine/Plugins/Online/OnlineSubsystemEOS/Source/OnlineSubsystemEOS/Private/OnlineTitleFileEOS.h b/Engine/Plugins/Online/OnlineSubsystemEOS/Source/OnlineSubsystemEOS/Private/OnlineTitleFileEOS.h
--- a/Engine/Plugins/Online/OnlineSubsystemEOS/Source/OnlineSubsystemEOS/Private/OnlineTitleFileEOS.h
+++ b/Engine/Plugins/Online/OnlineSubsystemEOS/Source/OnlineSubsystemEOS/Private/OnlineTitleFileEOS.h
@@ -5,7 +5,6 @@
 #include "CoreMinimal.h"
 #include "Online/CoreOnline.h"
 #include "Interfaces/OnlineTitleFileInterface.h"
-#include "OnlineSubsystemEOSPackage.h"
 #include "OnlineSubsystemEOSTypes.h"
 
 class FOnlineSubsystemEOS;
@@ -49,7 +48,6 @@ public:
 	virtual bool ReadFile(const FString& FileName) override;
 //~IOnlineTitleFile
 
-PACKAGE_SCOPE:
 	FOnlineTitleFileEOS(FOnlineSubsystemEOS* InSubsystem)
 		: EOSSubsystem(InSubsystem)
 	{
diff --git a/Engine/Plugins/Online/OnlineSubsystemEOS/Source/OnlineSubsystemEOS/Private/OnlineUserCloudEOS.cpp b/Engine/Plugins/Online/OnlineSubsystemEOS/Source/OnlineSubsystemEOS/Private/OnlineUserCloudEOS.cpp
--- a/Engine/Plugins/Online/OnlineSubsystemEOS/Source/OnlineSubsystemEOS/Private/OnlineUserCloudEOS.cpp
+++ b/Engine/Plugins/Online/OnlineSubsystemEOS/Source/OnlineSubsystemEOS/Private/OnlineUserCloudEOS.cpp
@@ -3,6 +3,7 @@
 #include "OnlineUserCloudEOS.h"
 #include "OnlineSubsystem.h"
 #include "OnlineSubsystemEOS.h"
+#include "OnlineSubsystemEOSPrivate.h"
 #include "OnlineSubsystemEOSTypes.h"
 #include "UserManagerEOS.h"
 #include "EOSSettings.h"
diff --git a/Engine/Plugins/Online/OnlineSubsystemEOS/Source/OnlineSubsystemEOS/Private/OnlineUserCloudEOS.h b/Engine/Plugins/Online/OnlineSubsystemEOS/Source/OnlineSubsystemEOS/Private/OnlineUserCloudEOS.h
--- a/Engine/Plugins/Online/OnlineSubsystemEOS/Source/OnlineSubsystemEOS/Private/OnlineUserCloudEOS.h
+++ b/Engine/Plugins/Online/OnlineSubsystemEOS/Source/OnlineSubsystemEOS/Private/OnlineUserCloudEOS.h
@@ -5,7 +5,6 @@
 #include "CoreMinimal.h"
 #include "Online/CoreOnline.h"
 #include "Interfaces/OnlineUserCloudInterface.h"
-#include "OnlineSubsystemEOSPackage.h"
 #include "OnlineSubsystemEOSTypes.h"
 
 class FOnlineSubsystemEOS;
@@ -54,7 +53,6 @@ public:
 	virtual void DumpCloudFileState(const FUniqueNetId& UserId, const FString& FileName) override;
 //~ IOnlineUserCloud
 
-PACKAGE_SCOPE:
 	FOnlineUserCloudEOS(FOnlineSubsystemEOS* InSubsystem)
 		: EOSSubsystem(InSubsystem)
 	{
diff --git a/Engine/Plugins/Online/OnlineSubsystemEOS/Source/OnlineSubsystemEOS/Private/UserManagerEOS.cpp b/Engine/Plugins/Online/OnlineSubsystemEOS/Source/OnlineSubsystemEOS/Private/UserManagerEOS.cpp
--- a/Engine/Plugins/Online/OnlineSubsystemEOS/Source/OnlineSubsystemEOS/Private/UserManagerEOS.cpp
+++ b/Engine/Plugins/Online/OnlineSubsystemEOS/Source/OnlineSubsystemEOS/Private/UserManagerEOS.cpp
@@ -2,6 +2,7 @@
 
 #include "UserManagerEOS.h"
 #include "OnlineSubsystemEOS.h"
+#include "OnlineSubsystemEOSPrivate.h"
 #include "Misc/CommandLine.h"
 #include "Misc/Guid.h"
 #include "Misc/OutputDeviceRedirector.h"
diff --git a/Engine/Plugins/Online/OnlineSubsystemEOS/Source/OnlineSubsystemEOS/Private/UserManagerEOS.h b/Engine/Plugins/Online/OnlineSubsystemEOS/Source/OnlineSubsystemEOS/Private/UserManagerEOS.h
--- a/Engine/Plugins/Online/OnlineSubsystemEOS/Source/OnlineSubsystemEOS/Private/UserManagerEOS.h
+++ b/Engine/Plugins/Online/OnlineSubsystemEOS/Source/OnlineSubsystemEOS/Private/UserManagerEOS.h
@@ -337,7 +337,6 @@ public:
 	virtual FUniqueNetIdPtr GetExternalIdMapping(const FExternalIdQueryOptions& QueryOptions, const FString& ExternalId) override;
 // ~IOnlineUser Interface
 
-PACKAGE_SCOPE:
 	EOS_EpicAccountId GetLocalEpicAccountId(int32 LocalUserNum) const;
 	EOS_EpicAccountId GetLocalEpicAccountId() const;
 	EOS_ProductUserId GetLocalProductUserId(int32 LocalUserNum) const;
diff --git a/Engine/Plugins/Online/OnlineSubsystemEOS/Source/OnlineSubsystemEOS/Public/IOnlineSubsystemEOS.h b/Engine/Plugins/Online/OnlineSubsystemEOS/Source/OnlineSubsystemEOS/Public/IOnlineSubsystemEOS.h
new file mode 100644
--- /dev/null
+++ b/Engine/Plugins/Online/OnlineSubsystemEOS/Source/OnlineSubsystemEOS/Public/IOnlineSubsystemEOS.h
@@ -0,0 +1,23 @@
+// Copyright Epic Games, Inc. All Rights Reserved.
+
+#pragma once
+
+#include "OnlineSubsystemImpl.h"
+
+class FUniqueNetId;
+class IVoiceChatUser;
+using IEOSPlatformHandlePtr = TSharedPtr<class IEOSPlatformHandle, ESPMode::ThreadSafe>;
+
+/**
+ *	OnlineSubsystemEOS - Implementation of the online subsystem for EOS services
+ */
+class ONLINESUBSYSTEMEOS_API IOnlineSubsystemEOS : 
+	public FOnlineSubsystemImpl
+{
+public:
+	IOnlineSubsystemEOS(FName InSubsystemName, FName InInstanceName) : FOnlineSubsystemImpl(InSubsystemName, InInstanceName) {}
+	virtual ~IOnlineSubsystemEOS() = default;
+
+	virtual IVoiceChatUser* GetVoiceChatUserInterface(const FUniqueNetId& LocalUserId) = 0;
+	virtual IEOSPlatformHandlePtr GetEOSPlatformHandle() const = 0;
+};
\ No newline at end of file
diff --git a/Engine/Plugins/Online/OnlineSubsystemEOS/Source/OnlineSubsystemEOS/Public/OnlineSubsystemEOS.h b/Engine/Plugins/Online/OnlineSubsystemEOS/Source/OnlineSubsystemEOS/Public/OnlineSubsystemEOS.h
deleted file mode 100644
diff --git a/Engine/Plugins/Online/OnlineSubsystemEOS/Source/OnlineSubsystemEOS/Public/OnlineSubsystemEOSModule.h b/Engine/Plugins/Online/OnlineSubsystemEOS/Source/OnlineSubsystemEOS/Public/OnlineSubsystemEOSModule.h
deleted file mode 100644
diff --git a/Engine/Plugins/Online/OnlineSubsystemEOS/Source/OnlineSubsystemEOS/Public/OnlineSubsystemEOSPackage.h b/Engine/Plugins/Online/OnlineSubsystemEOS/Source/OnlineSubsystemEOS/Public/OnlineSubsystemEOSPackage.h
deleted file mode 100644
