#PR9061

Refactor of string conversion to avoid string content corruption.
Contributed by: Sijun Seong (seasonseong)

diff --git a/Engine/Plugins/Online/OnlineSubsystemEOS/Source/OnlineSubsystemEOS/Private/OnlineSessionEOS.cpp b/Engine/Plugins/Online/OnlineSubsystemEOS/Source/OnlineSubsystemEOS/Private/OnlineSessionEOS.cpp
--- a/Engine/Plugins/Online/OnlineSubsystemEOS/Source/OnlineSubsystemEOS/Private/OnlineSessionEOS.cpp
+++ b/Engine/Plugins/Online/OnlineSubsystemEOS/Source/OnlineSubsystemEOS/Private/OnlineSessionEOS.cpp
@@ -40,8 +40,7 @@ FString MakeStringFromAttributeValue(const EOS_Sessions_AttributeData* Attribute
 		}
 		case EOS_ESessionAttributeType::EOS_SAT_String:
 		{
-			FString Value = Attribute->Value.AsUtf8;
-			return FString::Printf(TEXT("%s"), *Value);
+			return FString(UTF8_TO_TCHAR(Attribute->Value.AsUtf8));
 		}
 	}
 	return TEXT("");
@@ -1981,7 +1980,7 @@ void FOnlineSessionEOS::CopyAttributes(EOS_HSessionDetails SessionHandle, FOnlin
 					}
 					case EOS_ESessionAttributeType::EOS_SAT_String:
 					{
-						Setting.Data.SetValue(ANSI_TO_TCHAR(Attribute->Data->Value.AsUtf8));
+						Setting.Data.SetValue(UTF8_TO_TCHAR(Attribute->Data->Value.AsUtf8));
 						break;
 					}
 				}
@@ -4106,7 +4105,7 @@ void FOnlineSessionEOS::CopyLobbyAttributes(EOS_HLobbyDetails LobbyDetailsHandle
 				}
 				case EOS_ESessionAttributeType::EOS_SAT_String:
 				{
-					Setting.Data.SetValue(ANSI_TO_TCHAR(Attribute->Data->Value.AsUtf8));
+					Setting.Data.SetValue(UTF8_TO_TCHAR(Attribute->Data->Value.AsUtf8));
 					break;
 				}
 				}
@@ -4160,7 +4159,7 @@ void FOnlineSessionEOS::CopyLobbyMemberAttributes(EOS_HLobbyDetails LobbyDetails
 			}
 			case EOS_ESessionAttributeType::EOS_SAT_String:
 			{
-				Setting.Data.SetValue(ANSI_TO_TCHAR(Attribute->Data->Value.AsUtf8));
+				Setting.Data.SetValue(UTF8_TO_TCHAR(Attribute->Data->Value.AsUtf8));
 				break;
 			}
 			}
