Added missing delegate triggers for EOS Connect login

diff --git a/Engine/Plugins/Online/OnlineSubsystemEOS/Source/OnlineSubsystemEOS/Private/UserManagerEOS.cpp b/Engine/Plugins/Online/OnlineSubsystemEOS/Source/OnlineSubsystemEOS/Private/UserManagerEOS.cpp
--- a/Engine/Plugins/Online/OnlineSubsystemEOS/Source/OnlineSubsystemEOS/Private/UserManagerEOS.cpp
+++ b/Engine/Plugins/Online/OnlineSubsystemEOS/Source/OnlineSubsystemEOS/Private/UserManagerEOS.cpp
@@ -650,7 +650,9 @@ bool FUserManagerEOS::ConnectLoginNoEAS(int32 LocalUserNum)
 		{
 			if (!bWasSuccessful || !AuthToken.IsValid())
 			{
-				UE_LOG_ONLINE(Error, TEXT("ConnectLoginNoEAS(%d) failed due to the platform OSS giving an empty auth token"), LocalUserNum);
+				const FString ErrorString = FString::Printf(TEXT("ConnectLoginNoEAS(%d) failed due to the platform OSS giving an empty auth token"), LocalUserNum);
+				UE_LOG_ONLINE(Warning, TEXT("%s"), *ErrorString);
+				TriggerOnLoginCompleteDelegates(LocalUserNum, false, *FUniqueNetIdEOS::EmptyId(), ErrorString);
 				return;
 			}
 
@@ -685,7 +687,9 @@ bool FUserManagerEOS::ConnectLoginNoEAS(int32 LocalUserNum)
 				}
 				else
 				{
-					UE_LOG_ONLINE(Error, TEXT("ConnectLoginNoEAS(%d) failed with EOS result code (%s)"), LocalUserNum, ANSI_TO_TCHAR(EOS_EResult_ToString(Data->ResultCode)));
+					const FString ErrorString = FString::Printf(TEXT("ConnectLoginNoEAS(%d) failed with EOS result code (%s)"), LocalUserNum, ANSI_TO_TCHAR(EOS_EResult_ToString(Data->ResultCode)));
+					UE_LOG_ONLINE(Warning, TEXT("%s"), *ErrorString);
+					TriggerOnLoginCompleteDelegates(LocalUserNum, false, *FUniqueNetIdEOS::EmptyId(), ErrorString);
 				}
 			};
 			EOS_Connect_Login(EOSSubsystem->ConnectHandle, &Options, CallbackObj, CallbackObj->GetCallbackPtr());
