Refactor to avoid invalid reference to search result

diff --git a/Engine/Plugins/Online/OnlineSubsystemEOS/Source/OnlineSubsystemEOS/Private/OnlineSessionEOS.cpp b/Engine/Plugins/Online/OnlineSubsystemEOS/Source/OnlineSubsystemEOS/Private/OnlineSessionEOS.cpp
--- a/Engine/Plugins/Online/OnlineSubsystemEOS/Source/OnlineSubsystemEOS/Private/OnlineSessionEOS.cpp
+++ b/Engine/Plugins/Online/OnlineSubsystemEOS/Source/OnlineSubsystemEOS/Private/OnlineSessionEOS.cpp
@@ -3969,10 +3969,11 @@ void FOnlineSessionEOS::AddLobbySearchResult(EOS_HLobbyDetails LobbyDetailsHandl
 
 			EOS_ProductUserId TargetUserId = EOS_LobbyDetails_GetMemberByIndex(LobbyDetailsHandle, &GetMemberByIndexOptions);
 
-			EOSSubsystem->UserManager->GetEpicAccountIdAsync(TargetUserId, [this, LobbyDetailsHandle, &SearchResult](const EOS_ProductUserId& ProductUserId, EOS_EpicAccountId& EpicAccountId)
+			EOSSubsystem->UserManager->GetEpicAccountIdAsync(TargetUserId, [this, LobbyDetailsHandle, SearchSettings, Position](const EOS_ProductUserId& ProductUserId, EOS_EpicAccountId& EpicAccountId) 
 				{
 					FUniqueNetIdPtr UniqueNetId = EOSSubsystem->UserManager->CreateUniquePlayerId(MakeNetIdStringFromIds(EpicAccountId, ProductUserId));
 
+					FOnlineSessionSearchResult& SearchResult = SearchSettings->SearchResults[Position];
 					if (!SearchResult.Session.SessionSettings.MemberSettings.Contains(UniqueNetId->AsShared()))
 					{
 						SearchResult.Session.SessionSettings.MemberSettings.Add(UniqueNetId->AsShared(), FSessionSettings());
