Refactor of FUniqueNetIdEOS to hold member variables of the EOS Id types instead of strings, to streamline access to them and fix the referenced issue.
Refactor of callsites to use the new access methods, keeping existing logic.

diff --git a/Engine/Plugins/Online/OnlineSubsystemEOS/Source/OnlineSubsystemEOS/Private/NboSerializerEOS.h b/Engine/Plugins/Online/OnlineSubsystemEOS/Source/OnlineSubsystemEOS/Private/NboSerializerEOS.h
--- a/Engine/Plugins/Online/OnlineSubsystemEOS/Source/OnlineSubsystemEOS/Private/NboSerializerEOS.h
+++ b/Engine/Plugins/Online/OnlineSubsystemEOS/Source/OnlineSubsystemEOS/Private/NboSerializerEOS.h
@@ -44,8 +44,6 @@ public:
 	friend inline FNboSerializeToBufferEOS& operator<<(FNboSerializeToBufferEOS& Ar, const FUniqueNetIdEOS& UniqueId)
 	{
 		Ar << UniqueId.UniqueNetIdStr;
-		Ar << UniqueId.EpicAccountIdStr;
-		Ar << UniqueId.ProductUserIdStr;
 		return Ar;
 	}
 
@@ -92,8 +90,7 @@ public:
 	friend inline FNboSerializeFromBufferEOS& operator>>(FNboSerializeFromBufferEOS& Ar, FUniqueNetIdEOS& UniqueId)
 	{
 		Ar >> UniqueId.UniqueNetIdStr;
-		Ar >> UniqueId.EpicAccountIdStr;
-		Ar >> UniqueId.ProductUserIdStr;
+		UniqueId.ParseAccountIds();
 		return Ar;
 	}
 
diff --git a/Engine/Plugins/Online/OnlineSubsystemEOS/Source/OnlineSubsystemEOS/Private/OnlineLeaderboardsEOS.cpp b/Engine/Plugins/Online/OnlineSubsystemEOS/Source/OnlineSubsystemEOS/Private/OnlineLeaderboardsEOS.cpp
--- a/Engine/Plugins/Online/OnlineSubsystemEOS/Source/OnlineSubsystemEOS/Private/OnlineLeaderboardsEOS.cpp
+++ b/Engine/Plugins/Online/OnlineSubsystemEOS/Source/OnlineSubsystemEOS/Private/OnlineLeaderboardsEOS.cpp
@@ -93,7 +93,7 @@ bool FOnlineLeaderboardsEOS::ReadLeaderboards(const TArray<FUniqueNetIdRef>& Pla
 	for (const FUniqueNetIdRef& NetId : Players)
 	{
 		const FUniqueNetIdEOS& EOSId = FUniqueNetIdEOS::Cast(*NetId);
-		EOS_ProductUserId UserId = EOS_ProductUserId_FromString(TCHAR_TO_UTF8(*EOSId.ProductUserIdStr));
+		const EOS_ProductUserId UserId = EOSId.GetProductUserId();
 		if (UserId != nullptr)
 		{
 			ProductUserIds.Add(UserId);
@@ -134,7 +134,8 @@ bool FOnlineLeaderboardsEOS::ReadLeaderboards(const TArray<FUniqueNetIdRef>& Pla
 		for (FUniqueNetIdRef NetId : QueryContext->Players)
 		{
 			FString Nickname = EOSSubsystem->UserManager->GetPlayerNickname(*NetId);
-			EOS_ProductUserId UserId = EOSSubsystem->UserManager->GetProductUserId(*NetId);
+			const FUniqueNetIdEOS& EOSId = FUniqueNetIdEOS::Cast(*NetId);
+			const EOS_ProductUserId UserId = EOSId.GetProductUserId();
 			if (UserId == nullptr)
 			{
 				QueryContext->AddEmptyRowForPlayer(NetId, Nickname);
diff --git a/Engine/Plugins/Online/OnlineSubsystemEOS/Source/OnlineSubsystemEOS/Private/OnlineSessionEOS.cpp b/Engine/Plugins/Online/OnlineSubsystemEOS/Source/OnlineSubsystemEOS/Private/OnlineSessionEOS.cpp
--- a/Engine/Plugins/Online/OnlineSubsystemEOS/Source/OnlineSubsystemEOS/Private/OnlineSessionEOS.cpp
+++ b/Engine/Plugins/Online/OnlineSubsystemEOS/Source/OnlineSubsystemEOS/Private/OnlineSessionEOS.cpp
@@ -2431,10 +2431,12 @@ bool FOnlineSessionEOS::FindFriendSession(int32 LocalUserNum, const FUniqueNetId
 	EOS_EResult CreateLobbySearchResult = EOS_Lobby_CreateLobbySearch(LobbyHandle, &CreateLobbySearchOptions, &LobbySearchHandle);
 	if (CreateLobbySearchResult == EOS_EResult::EOS_Success)
 	{
+		const FUniqueNetIdEOS& FriendEOSId = FUniqueNetIdEOS::Cast(Friend);
+
 		// Set the user we wan to use to find lobbies
 		EOS_LobbySearch_SetTargetUserIdOptions SetTargetUserIdOptions = { 0 };
 		SetTargetUserIdOptions.ApiVersion = EOS_LOBBYSEARCH_SETTARGETUSERID_API_LATEST;
-		SetTargetUserIdOptions.TargetUserId = EOSSubsystem->UserManager->GetProductUserId(Friend);
+		SetTargetUserIdOptions.TargetUserId = FriendEOSId.GetProductUserId();
 
 		// TODO: Using this as a search parameter only works if we use the owner's id (search for lobbies we're already in). Pending API fix so it works with other users too.
 		EOS_LobbySearch_SetTargetUserId(LobbySearchHandle, &SetTargetUserIdOptions);
@@ -2572,32 +2574,35 @@ bool FOnlineSessionEOS::SendSessionInviteToFriend(int32 LocalUserNum, FName Sess
 		UE_LOG_ONLINE_SESSION(Error, TEXT("SendSessionInviteToFriend() failed due to user (%d) being not logged in"), (int32)LocalUserNum);
 		return false;
 	}
-	EOS_ProductUserId TargetUserId = EOSSubsystem->UserManager->GetProductUserId(Friend);
-	if (TargetUserId == nullptr)
+	const FUniqueNetIdEOS& FriendEOSId = FUniqueNetIdEOS::Cast(Friend);
+	const EOS_ProductUserId FriendId = FriendEOSId.GetProductUserId();
+	if (EOS_ProductUserId_IsValid(FriendId) == EOS_FALSE)
 	{
 		UE_LOG_ONLINE_SESSION(Error, TEXT("SendSessionInviteToFriend() failed due to target user (%s) having not played this game"), *Friend.ToDebugString());
 		return false;
 	}
 
-	return SendSessionInvite(SessionName, LocalUserId, TargetUserId);
+	return SendSessionInvite(SessionName, LocalUserId, FriendId);
 };
 
 bool FOnlineSessionEOS::SendSessionInviteToFriend(const FUniqueNetId& LocalNetId, FName SessionName, const FUniqueNetId& Friend)
 {
-	EOS_ProductUserId LocalUserId = EOSSubsystem->UserManager->GetProductUserId(LocalNetId);
-	if (LocalUserId == nullptr)
+	const FUniqueNetIdEOS& LocalEOSId = FUniqueNetIdEOS::Cast(LocalNetId);
+	const EOS_ProductUserId LocalUserId = LocalEOSId.GetProductUserId();
+	if (EOS_ProductUserId_IsValid(LocalUserId) == EOS_FALSE)
 	{
 		UE_LOG_ONLINE_SESSION(Error, TEXT("SendSessionInviteToFriend() failed due to user (%s) being not logged in"), *LocalNetId.ToDebugString());
 		return false;
 	}
-	EOS_ProductUserId TargetUserId = EOSSubsystem->UserManager->GetProductUserId(Friend);
-	if (TargetUserId == nullptr)
+	const FUniqueNetIdEOS& FriendEOSId = FUniqueNetIdEOS::Cast(Friend);
+	const EOS_ProductUserId FriendId = FriendEOSId.GetProductUserId();
+	if (EOS_ProductUserId_IsValid(FriendId) == EOS_FALSE)
 	{
 		UE_LOG_ONLINE_SESSION(Error, TEXT("SendSessionInviteToFriend() failed due to target user (%s) having not played this game"), *Friend.ToDebugString());
 		return false;
 	}
 
-	return SendSessionInvite(SessionName, LocalUserId, TargetUserId);
+	return SendSessionInvite(SessionName, LocalUserId, FriendId);
 }
 
 bool FOnlineSessionEOS::SendSessionInviteToFriends(int32 LocalUserNum, FName SessionName, const TArray< FUniqueNetIdRef >& Friends)
@@ -2750,6 +2755,7 @@ bool FOnlineSessionEOS::RegisterPlayers(FName SessionName, const TArray< FUnique
 		for (int32 PlayerIdx=0; PlayerIdx<Players.Num(); PlayerIdx++)
 		{
 			const FUniqueNetIdRef& PlayerId = Players[PlayerIdx];
+			const FUniqueNetIdEOS& PlayerEOSId = FUniqueNetIdEOS::Cast(*PlayerId);
 
 			FUniqueNetIdMatcher PlayerMatch(*PlayerId);
 			if (Session->RegisteredPlayers.IndexOfByPredicate(PlayerMatch) == INDEX_NONE)
@@ -2757,7 +2763,7 @@ bool FOnlineSessionEOS::RegisterPlayers(FName SessionName, const TArray< FUnique
 				Session->RegisteredPlayers.Add(PlayerId);
 				if (bRegisterEOS)
 				{
-					EOSIds.Add(EOSSubsystem->UserManager->GetProductUserId(*PlayerId));
+					EOSIds.Add(PlayerEOSId.GetProductUserId());
 				}
 
 				// update number of open connections
@@ -2833,6 +2839,7 @@ bool FOnlineSessionEOS::UnregisterPlayers(FName SessionName, const TArray< FUniq
 		for (int32 PlayerIdx=0; PlayerIdx < Players.Num(); PlayerIdx++)
 		{
 			const FUniqueNetIdRef& PlayerId = Players[PlayerIdx];
+			const FUniqueNetIdEOS& PlayerEOSId = FUniqueNetIdEOS::Cast(*PlayerId);
 
 			FUniqueNetIdMatcher PlayerMatch(*PlayerId);
 			int32 RegistrantIndex = Session->RegisteredPlayers.IndexOfByPredicate(PlayerMatch);
@@ -2841,7 +2848,7 @@ bool FOnlineSessionEOS::UnregisterPlayers(FName SessionName, const TArray< FUniq
 				Session->RegisteredPlayers.RemoveAtSwap(RegistrantIndex);
 				if (bUnregisterEOS)
 				{
-					EOSIds.Add(EOSSubsystem->UserManager->GetProductUserId(*PlayerId));
+					EOSIds.Add(PlayerEOSId.GetProductUserId());
 				}
 
 				// update number of open connections
@@ -3219,12 +3226,14 @@ void FOnlineSessionEOS::RemovePlayerFromSession(int32 LocalUserNum, FName Sessio
 	FNamedOnlineSession* Session = GetNamedSession(SessionName);
 	if (Session)
 	{
+		const FUniqueNetIdEOS& TargetPlayerEOSId = FUniqueNetIdEOS::Cast(TargetPlayerId);
+
 		EOS_Lobby_KickMemberOptions KickMemberOptions = {};
 		KickMemberOptions.ApiVersion = EOS_LOBBY_KICKMEMBER_API_LATEST;
 		const FTCHARToUTF8 Utf8LobbyId(*Session->SessionInfo->GetSessionId().ToString());
 		KickMemberOptions.LobbyId = (EOS_LobbyId)Utf8LobbyId.Get();
 		KickMemberOptions.LocalUserId = EOSSubsystem->UserManager->GetLocalProductUserId(LocalUserNum);
-		KickMemberOptions.TargetUserId = EOSSubsystem->UserManager->GetProductUserId(TargetPlayerId);
+		KickMemberOptions.TargetUserId = TargetPlayerEOSId.GetProductUserId();
 
 		FLobbyRemovePlayerCallback* CallbackObj = new FLobbyRemovePlayerCallback(FOnlineSessionEOSWeakPtr(AsShared()));
 		CallbackObj->CallbackLambda = [this](const EOS_Lobby_KickMemberCallbackInfo* Data)
diff --git a/Engine/Plugins/Online/OnlineSubsystemEOS/Source/OnlineSubsystemEOS/Private/OnlineStatsEOS.cpp b/Engine/Plugins/Online/OnlineSubsystemEOS/Source/OnlineSubsystemEOS/Private/OnlineStatsEOS.cpp
--- a/Engine/Plugins/Online/OnlineSubsystemEOS/Source/OnlineSubsystemEOS/Private/OnlineStatsEOS.cpp
+++ b/Engine/Plugins/Online/OnlineSubsystemEOS/Source/OnlineSubsystemEOS/Private/OnlineStatsEOS.cpp
@@ -108,7 +108,7 @@ void AppendStats(TUniqueNetIdMap<TSharedRef<FOnlineStatsUserStats>>& StatsCache,
 void FOnlineStatsEOS::QueryStats(const FUniqueNetIdRef LocalUserId, const TArray<FUniqueNetIdRef>& StatUsers, const TArray<FString>& StatNames, const FOnlineStatsQueryUsersStatsComplete& Delegate)
 {
 	const FUniqueNetIdEOS& LocalEOSId = FUniqueNetIdEOS::Cast(*LocalUserId);
-	EOS_ProductUserId LocalEOSUserId = EOS_ProductUserId_FromString(TCHAR_TO_UTF8(*LocalEOSId.ProductUserIdStr));
+	const EOS_ProductUserId LocalEOSUserId = LocalEOSId.GetProductUserId();
 	if (StatNames.Num() == 0)
 	{
 		UE_LOG_ONLINE_STATS(Warning, TEXT("QueryStats() without a list of stats names to query is not supported"));
@@ -134,7 +134,7 @@ void FOnlineStatsEOS::QueryStats(const FUniqueNetIdRef LocalUserId, const TArray
 	for (const FUniqueNetIdRef& StatUserId : StatUsers)
 	{
 		const FUniqueNetIdEOS& EOSId = FUniqueNetIdEOS::Cast(*StatUserId);
-		EOS_ProductUserId UserId = EOS_ProductUserId_FromString(TCHAR_TO_UTF8(*EOSId.ProductUserIdStr));
+		const EOS_ProductUserId UserId = EOSId.GetProductUserId();
 		if (UserId == nullptr)
 		{
 			continue;
@@ -148,7 +148,7 @@ void FOnlineStatsEOS::QueryStats(const FUniqueNetIdRef LocalUserId, const TArray
 	for (const FUniqueNetIdRef& StatUserId : StatUsers)
 	{
 		const FUniqueNetIdEOS& EOSId = FUniqueNetIdEOS::Cast(*StatUserId);
-		const EOS_ProductUserId TargetEOSUserId = EOS_ProductUserId_FromString(TCHAR_TO_UTF8(*EOSId.ProductUserIdStr));
+		const EOS_ProductUserId TargetEOSUserId = EOSId.GetProductUserId();
 		if (TargetEOSUserId == nullptr)
 		{
 			continue;
@@ -325,7 +325,7 @@ void FOnlineStatsEOS::WriteStats(EOS_ProductUserId LocalUserId, EOS_ProductUserI
 void FOnlineStatsEOS::UpdateStats(const FUniqueNetIdRef LocalUserId, const TArray<FOnlineStatsUserUpdatedStats>& UpdatedUserStats, const FOnlineStatsUpdateStatsComplete& Delegate)
 {
 	const FUniqueNetIdEOS& EOSId = FUniqueNetIdEOS::Cast(*LocalUserId);
-	EOS_ProductUserId UserId = EOS_ProductUserId_FromString(TCHAR_TO_UTF8(*EOSId.ProductUserIdStr));
+	const EOS_ProductUserId UserId = EOSId.GetProductUserId();
 	if (UserId == nullptr)
 	{
 		UE_LOG_ONLINE_STATS(Error, TEXT("UpdateStats() failed for unknown player (%s)"), *EOSId.UniqueNetIdStr);
@@ -342,7 +342,8 @@ void FOnlineStatsEOS::UpdateStats(const FUniqueNetIdRef LocalUserId, const TArra
 	// Trigger one write for each user
 	for (const FOnlineStatsUserUpdatedStats& StatsUpdate : UpdatedUserStats)
 	{
-		EOS_ProductUserId StatsUser = EOSSubsystem->UserManager->GetProductUserId(*StatsUpdate.Account);
+		const FUniqueNetIdEOS& AccountEOSId = FUniqueNetIdEOS::Cast(*StatsUpdate.Account);
+		const EOS_ProductUserId StatsUser = AccountEOSId.GetProductUserId();
 		if (StatsUser != nullptr)
 		{
 			WriteStats(UserId, StatsUser, StatsUpdate);
diff --git a/Engine/Plugins/Online/OnlineSubsystemEOS/Source/OnlineSubsystemEOS/Private/OnlineStoreEOS.cpp b/Engine/Plugins/Online/OnlineSubsystemEOS/Source/OnlineSubsystemEOS/Private/OnlineStoreEOS.cpp
--- a/Engine/Plugins/Online/OnlineSubsystemEOS/Source/OnlineSubsystemEOS/Private/OnlineStoreEOS.cpp
+++ b/Engine/Plugins/Online/OnlineSubsystemEOS/Source/OnlineSubsystemEOS/Private/OnlineStoreEOS.cpp
@@ -47,7 +47,8 @@ void FOnlineStoreEOS::QueryOffers(const FUniqueNetId& UserId, const FOnQueryOnli
 		Delegate.ExecuteIfBound(true, CachedOfferIds, TEXT("Returning cached offers"));
 		return;
 	}
-	EOS_EpicAccountId AccountId = EOSSubsystem->UserManager->GetEpicAccountId(UserId);
+	const FUniqueNetIdEOS& UserEOSId = FUniqueNetIdEOS::Cast(UserId);
+	const EOS_EpicAccountId AccountId = UserEOSId.GetEpicAccountId();
 	if (AccountId == nullptr)
 	{
 		Delegate.ExecuteIfBound(false, TArray<FUniqueOfferId>(), TEXT("Can't query offers for a null user"));
@@ -143,7 +144,8 @@ typedef TEOSCallback<EOS_Ecom_OnCheckoutCallback, EOS_Ecom_CheckoutCallbackInfo,
 
 void FOnlineStoreEOS::Checkout(const FUniqueNetId& UserId, const FPurchaseCheckoutRequest& CheckoutRequest, const FOnPurchaseCheckoutComplete& Delegate)
 {
-	EOS_EpicAccountId AccountId = EOSSubsystem->UserManager->GetEpicAccountId(UserId);
+	const FUniqueNetIdEOS& UserEOSId = FUniqueNetIdEOS::Cast(UserId);
+	const EOS_EpicAccountId AccountId = UserEOSId.GetEpicAccountId();
 	if (AccountId == nullptr)
 	{
 		UE_LOG_ONLINE(Error, TEXT("Checkout: failed due to invalid user"));
@@ -237,7 +239,8 @@ typedef TEOSCallback<EOS_Ecom_OnQueryEntitlementsCallback, EOS_Ecom_QueryEntitle
 
 void FOnlineStoreEOS::QueryReceipts(const FUniqueNetId& UserId, bool bRestoreReceipts, const FOnQueryReceiptsComplete& Delegate)
 {
-	EOS_EpicAccountId AccountId = EOSSubsystem->UserManager->GetEpicAccountId(UserId);
+	const FUniqueNetIdEOS& UserEOSId = FUniqueNetIdEOS::Cast(UserId);
+	const EOS_EpicAccountId AccountId = UserEOSId.GetEpicAccountId();
 	if (AccountId == nullptr)
 	{
 		UE_LOG_ONLINE(Error, TEXT("QueryReceipts: failed due to invalid user"));
@@ -312,7 +315,8 @@ typedef TEOSCallback<EOS_Ecom_OnRedeemEntitlementsCallback, EOS_Ecom_RedeemEntit
 
 void FOnlineStoreEOS::FinalizeReceiptValidationInfo(const FUniqueNetId& UserId, FString& InReceiptValidationInfo, const FOnFinalizeReceiptValidationInfoComplete& Delegate)
 {
-	EOS_EpicAccountId AccountId = EOSSubsystem->UserManager->GetEpicAccountId(UserId);
+	const FUniqueNetIdEOS& UserEOSId = FUniqueNetIdEOS::Cast(UserId);
+	const EOS_EpicAccountId AccountId = UserEOSId.GetEpicAccountId();
 	if (AccountId == nullptr)
 	{
 		Delegate.ExecuteIfBound(ONLINE_ERROR(EOnlineErrorResult::InvalidUser), InReceiptValidationInfo);
diff --git a/Engine/Plugins/Online/OnlineSubsystemEOS/Source/OnlineSubsystemEOS/Private/OnlineSubsystemEOS.cpp b/Engine/Plugins/Online/OnlineSubsystemEOS/Source/OnlineSubsystemEOS/Private/OnlineSubsystemEOS.cpp
--- a/Engine/Plugins/Online/OnlineSubsystemEOS/Source/OnlineSubsystemEOS/Private/OnlineSubsystemEOS.cpp
+++ b/Engine/Plugins/Online/OnlineSubsystemEOS/Source/OnlineSubsystemEOS/Private/OnlineSubsystemEOS.cpp
@@ -676,7 +676,7 @@ IVoiceChatUser* FOnlineSubsystemEOS::GetVoiceChatUserInterface(const FUniqueNetI
 		else
 		{
 			FEOSVoiceChatUser* VoiceChatUser = static_cast<FEOSVoiceChatUser*>(VoiceChatInterface->CreateUser());
-			VoiceChatUser->Login(UserManager->GetPlatformUserIdFromUniqueNetId(LocalUserId), FUniqueNetIdEOS::Cast(LocalUserId).ProductUserIdStr, FString(), FOnVoiceChatLoginCompleteDelegate());
+			VoiceChatUser->Login(UserManager->GetPlatformUserIdFromUniqueNetId(LocalUserId), LexToString(FUniqueNetIdEOS::Cast(LocalUserId).GetProductUserId()), FString(), FOnVoiceChatLoginCompleteDelegate());
 
 			const FOnlineSubsystemEOSVoiceChatUserWrapperRef& Wrapper = LocalVoiceChatUsers.Emplace(LocalUserId.AsShared(), MakeShared<FOnlineSubsystemEOSVoiceChatUserWrapper, ESPMode::ThreadSafe>(*VoiceChatUser));
 			Result = &Wrapper.Get();
diff --git a/Engine/Plugins/Online/OnlineSubsystemEOS/Source/OnlineSubsystemEOS/Private/OnlineSubsystemEOSTypes.h b/Engine/Plugins/Online/OnlineSubsystemEOS/Source/OnlineSubsystemEOS/Private/OnlineSubsystemEOSTypes.h
--- a/Engine/Plugins/Online/OnlineSubsystemEOS/Source/OnlineSubsystemEOS/Private/OnlineSubsystemEOSTypes.h
+++ b/Engine/Plugins/Online/OnlineSubsystemEOS/Source/OnlineSubsystemEOS/Private/OnlineSubsystemEOSTypes.h
@@ -78,6 +78,16 @@ public:
 		return EOS_ID_BYTE_SIZE;
 	}
 
+	const EOS_EpicAccountId GetEpicAccountId() const
+	{
+		return EpicAccountId;
+	}
+
+	const EOS_ProductUserId GetProductUserId() const
+	{
+		return ProductUserId;
+	}
+
 	void UpdateNetIdStr(const FString& InNetIdStr)
 	{
 		UniqueNetIdStr = InNetIdStr;
@@ -88,22 +98,28 @@ public:
 	{
 		TArray<FString> AccountIds;
 		UniqueNetIdStr.ParseIntoArray(AccountIds, EOS_ID_SEPARATOR, false);
+
+		FString EpicAccountIdStr = EMPTY_EASID;
 		if (AccountIds.Num() > 0 && AccountIds[0].Len() > 0)
 		{
 			EpicAccountIdStr = AccountIds[0];
+			EpicAccountId = EOS_EpicAccountId_FromString(TCHAR_TO_UTF8(*EpicAccountIdStr));
 		}
 		else
 		{
-			EpicAccountIdStr = EMPTY_EASID;
+			EpicAccountId = nullptr;
 		}
 		AddToBuffer(RawBytes, EpicAccountIdStr);
+
+		FString ProductUserIdStr = EMPTY_PUID;
 		if (AccountIds.Num() > 1 && AccountIds[1].Len() > 0)
 		{
 			ProductUserIdStr = AccountIds[1];
+			ProductUserId = EOS_ProductUserId_FromString(TCHAR_TO_UTF8(*ProductUserIdStr));
 		}
 		else
 		{
-			ProductUserIdStr = EMPTY_PUID;
+			ProductUserId = nullptr;
 		}
 		AddToBuffer(RawBytes + ID_HALF_BYTE_SIZE, ProductUserIdStr);
 	}
@@ -120,10 +136,11 @@ public:
 		}
 	}
 
-	FString EpicAccountIdStr;
-	FString ProductUserIdStr;
-	uint8 RawBytes[EOS_ID_BYTE_SIZE] = { 0 };
 private:
+	EOS_EpicAccountId EpicAccountId;
+	EOS_ProductUserId ProductUserId;
+	uint8 RawBytes[EOS_ID_BYTE_SIZE] = { 0 };
+
 	FUniqueNetIdEOS()
 		PRAGMA_DISABLE_DEPRECATION_WARNINGS
 		: FUniqueNetIdString(EMPTY_EASID EOS_ID_SEPARATOR EMPTY_PUID)
@@ -138,9 +155,16 @@ private:
 		PRAGMA_ENABLE_DEPRECATION_WARNINGS
 	{
 		check(Size == EOS_ID_BYTE_SIZE);
-		EpicAccountIdStr = BytesToHex(Bytes, ID_HALF_BYTE_SIZE);
-		ProductUserIdStr = BytesToHex(Bytes + ID_HALF_BYTE_SIZE, ID_HALF_BYTE_SIZE);
+		const FString EpicAccountIdStr = BytesToHex(Bytes, ID_HALF_BYTE_SIZE);
+		EpicAccountId = EOS_EpicAccountId_FromString(TCHAR_TO_UTF8(*EpicAccountIdStr));
+		AddToBuffer(RawBytes, EpicAccountIdStr);
+
+		const FString ProductUserIdStr = BytesToHex(Bytes + ID_HALF_BYTE_SIZE, ID_HALF_BYTE_SIZE);
+		ProductUserId = EOS_ProductUserId_FromString(TCHAR_TO_UTF8(*ProductUserIdStr));
+		AddToBuffer(RawBytes + ID_HALF_BYTE_SIZE, ProductUserIdStr);
+
 		UniqueNetIdStr = EpicAccountIdStr + EOS_ID_SEPARATOR + ProductUserIdStr;
+
 		Type = FName("EOS");
 	}
 
diff --git a/Engine/Plugins/Online/OnlineSubsystemEOS/Source/OnlineSubsystemEOS/Private/OnlineUserCloudEOS.cpp b/Engine/Plugins/Online/OnlineSubsystemEOS/Source/OnlineSubsystemEOS/Private/OnlineUserCloudEOS.cpp
--- a/Engine/Plugins/Online/OnlineSubsystemEOS/Source/OnlineSubsystemEOS/Private/OnlineUserCloudEOS.cpp
+++ b/Engine/Plugins/Online/OnlineSubsystemEOS/Source/OnlineSubsystemEOS/Private/OnlineUserCloudEOS.cpp
@@ -177,7 +177,8 @@ bool FOnlineUserCloudEOS::ClearFile(const FUniqueNetId& UserId, const FString& F
 
 void FOnlineUserCloudEOS::EnumerateUserFiles(const FUniqueNetId& UserId)
 {
-	EOS_ProductUserId LocalUserId = EOSSubsystem->UserManager->GetProductUserId(UserId);
+	const FUniqueNetIdEOS& UserEOSId = FUniqueNetIdEOS::Cast(UserId);
+	const EOS_ProductUserId LocalUserId = UserEOSId.GetProductUserId();
 
 	if (LocalUserId == nullptr)
 	{
@@ -419,9 +420,11 @@ bool FOnlineUserCloudEOS::ReadUserFile(const FUniqueNetId& UserId, const FString
 		TriggerOnReadUserFileCompleteDelegates(bWasSuccessful, *SharedUserId, FileName);
 	};
 
+	const FUniqueNetIdEOS& UserEOSId = FUniqueNetIdEOS::Cast(UserId);
+
 	EOS_PlayerDataStorage_ReadFileOptions ReadFileOptions = {};
 	ReadFileOptions.ApiVersion = EOS_PLAYERDATASTORAGE_READFILEOPTIONS_API_LATEST;
-	ReadFileOptions.LocalUserId = EOSSubsystem->UserManager->GetProductUserId(UserId);
+	ReadFileOptions.LocalUserId = UserEOSId.GetProductUserId();
 	ReadFileOptions.Filename = FileNameUtf8.Get();
 	ReadFileOptions.ReadChunkLengthBytes = (uint32_t)ReadChunkSize;
 	ReadFileOptions.ReadFileDataCallback = CallbackObj->GetNested1CallbackPtr();
@@ -595,9 +598,11 @@ bool FOnlineUserCloudEOS::WriteUserFile(const FUniqueNetId& UserId, const FStrin
 		TriggerOnWriteUserFileCompleteDelegates(bWasSuccessful, *SharedUserId, FileName);
 	};
 
+	const FUniqueNetIdEOS& UserEOSId = FUniqueNetIdEOS::Cast(UserId);
+
 	EOS_PlayerDataStorage_WriteFileOptions WriteFileOptions = {};
 	WriteFileOptions.ApiVersion = EOS_PLAYERDATASTORAGE_WRITEFILEOPTIONS_API_LATEST;
-	WriteFileOptions.LocalUserId = EOSSubsystem->UserManager->GetProductUserId(UserId);
+	WriteFileOptions.LocalUserId = UserEOSId.GetProductUserId();
 	WriteFileOptions.Filename = FileNameUtf8.Get();
 	WriteFileOptions.ChunkLengthBytes = (uint32_t)ReadChunkSize;
 	WriteFileOptions.WriteFileDataCallback = CallbackObj->GetNested1CallbackPtr();
@@ -747,9 +752,11 @@ bool FOnlineUserCloudEOS::DeleteUserFile(const FUniqueNetId& UserId, const FStri
 	{
 		FTCHARToUTF8 FileNameUtf8(*FileName);
 
+		const FUniqueNetIdEOS& UserEOSId = FUniqueNetIdEOS::Cast(UserId);
+
 		EOS_PlayerDataStorage_DeleteFileOptions Options = {};
 		Options.ApiVersion = EOS_PLAYERDATASTORAGE_DELETEFILEOPTIONS_API_LATEST;
-		Options.LocalUserId = EOSSubsystem->UserManager->GetProductUserId(UserId);
+		Options.LocalUserId = UserEOSId.GetProductUserId();
 		Options.Filename = FileNameUtf8.Get();
 
 		FOnDeleteFileCallback* CallbackObj = new FOnDeleteFileCallback(FOnlineUserCloudEOSWeakPtr(AsShared()));
diff --git a/Engine/Plugins/Online/OnlineSubsystemEOS/Source/OnlineSubsystemEOS/Private/UserManagerEOS.cpp b/Engine/Plugins/Online/OnlineSubsystemEOS/Source/OnlineSubsystemEOS/Private/UserManagerEOS.cpp
--- a/Engine/Plugins/Online/OnlineSubsystemEOS/Source/OnlineSubsystemEOS/Private/UserManagerEOS.cpp
+++ b/Engine/Plugins/Online/OnlineSubsystemEOS/Source/OnlineSubsystemEOS/Private/UserManagerEOS.cpp
@@ -974,7 +974,7 @@ bool FUserManagerEOS::Logout(int32 LocalUserNum)
 
 	EOS_Auth_LogoutOptions LogoutOptions = { };
 	LogoutOptions.ApiVersion = EOS_AUTH_LOGOUT_API_LATEST;
-	LogoutOptions.LocalUserId = StringToAccountIdMap[UserId->UniqueNetIdStr];
+	LogoutOptions.LocalUserId = UserId->GetEpicAccountId();
 
 	EOS_Auth_Logout(EOSSubsystem->AuthHandle, &LogoutOptions, CallbackObj, CallbackObj->GetCallbackPtr());
 
@@ -1026,11 +1026,9 @@ void FUserManagerEOS::AddLocalUser(int32 LocalUserNum, EOS_EpicAccountId EpicAcc
 	StringToUserAccountMap.Emplace(NetId, UserAccountRef);
 	AccountIdToStringMap.Emplace(EpicAccountId, NetId);
 	ProductUserIdToStringMap.Emplace(UserId, *NetId);
-	StringToAccountIdMap.Emplace(NetId, EpicAccountId);
 	EpicAccountIdToAttributeAccessMap.Emplace(EpicAccountId, UserAccountRef);
 	UserNumToProductUserIdMap.Emplace(LocalUserNum, UserId);
 	ProductUserIdToUserNumMap.Emplace(UserId, LocalUserNum);
-	StringToProductUserIdMap.Emplace(NetId, UserId);
 
 	// Init player lists
 	FFriendsListEOSRef FriendsList = MakeShareable(new FFriendsListEOS(LocalUserNum, UserNetId));
@@ -1136,14 +1134,19 @@ FUniqueNetIdPtr FUserManagerEOS::GetUniquePlayerId(int32 LocalUserNum) const
 int32 FUserManagerEOS::GetLocalUserNumFromUniqueNetId(const FUniqueNetId& NetId) const
 {
 	const FUniqueNetIdEOS& EosId = FUniqueNetIdEOS::Cast(NetId);
-	if (StringToAccountIdMap.Contains(EosId.UniqueNetIdStr))
+
+	const EOS_EpicAccountId AccountId = EosId.GetEpicAccountId();
+	if (const int32* UserNum = AccountIdToUserNumMap.Find(AccountId))
 	{
-		EOS_EpicAccountId AccountId = StringToAccountIdMap[EosId.UniqueNetIdStr];
-		if (AccountIdToUserNumMap.Contains(AccountId))
-		{
-			return AccountIdToUserNumMap[AccountId];
-		}
+		return *UserNum;
 	}
+
+	const EOS_ProductUserId ProductUserId = EosId.GetProductUserId();
+	if (const int32* UserNum = ProductUserIdToUserNumMap.Find(ProductUserId))
+	{
+		return *UserNum;
+	}
+
 	// Use the default user if we can't find the person that they want
 	return DefaultLocalUser;
 }
@@ -1151,7 +1154,7 @@ int32 FUserManagerEOS::GetLocalUserNumFromUniqueNetId(const FUniqueNetId& NetId)
 bool FUserManagerEOS::IsLocalUser(const FUniqueNetId& NetId) const
 {
 	const FUniqueNetIdEOS& EosId = FUniqueNetIdEOS::Cast(NetId);
-	return StringToAccountIdMap.Contains(EosId.UniqueNetIdStr);
+	return AccountIdToUserNumMap.Contains(EosId.GetEpicAccountId()) || ProductUserIdToUserNumMap.Contains(EosId.GetProductUserId());
 }
 
 FUniqueNetIdEOSPtr FUserManagerEOS::GetLocalUniqueNetIdEOS(int32 LocalUserNum) const
@@ -1228,28 +1231,6 @@ EOS_ProductUserId FUserManagerEOS::GetLocalProductUserId(EOS_EpicAccountId Accou
 	return nullptr;
 }
 
-EOS_EpicAccountId FUserManagerEOS::GetEpicAccountId(const FUniqueNetId& NetId) const
-{
-	const FUniqueNetIdEOS& EOSId = FUniqueNetIdEOS::Cast(NetId);
-
-	if (StringToAccountIdMap.Contains(EOSId.UniqueNetIdStr))
-	{
-		return StringToAccountIdMap[EOSId.UniqueNetIdStr];
-	}
-	return nullptr;
-}
-
-EOS_ProductUserId FUserManagerEOS::GetProductUserId(const FUniqueNetId& NetId) const
-{
-	const FUniqueNetIdEOS& EOSId = FUniqueNetIdEOS::Cast(NetId);
-
-	if (StringToProductUserIdMap.Contains(EOSId.UniqueNetIdStr))
-	{
-		return StringToProductUserIdMap[EOSId.UniqueNetIdStr];
-	}
-	return nullptr;
-}
-
 typedef TEOSCallback<EOS_Connect_OnQueryProductUserIdMappingsCallback, EOS_Connect_QueryProductUserIdMappingsCallbackInfo, FUserManagerEOS> FConnectQueryProductUserIdMappingsCallback;
 
 /**
@@ -1377,11 +1358,10 @@ void FUserManagerEOS::RemoveLocalUser(int32 LocalUserNum)
 		EOSSubsystem->ReleaseVoiceChatUserInterface(**FoundId);
 		LocalUserNumToFriendsListMap.Remove(LocalUserNum);
 		const FString& NetId = (*FoundId)->UniqueNetIdStr;
-		EOS_EpicAccountId AccountId = StringToAccountIdMap[NetId];
+		const EOS_EpicAccountId AccountId = (*FoundId)->GetEpicAccountId();
 		AccountIdToStringMap.Remove(AccountId);
 		AccountIdToUserNumMap.Remove(AccountId);
 		NetIdStringToOnlineUserMap.Remove(NetId);
-		StringToAccountIdMap.Remove(NetId);
 		StringToUserAccountMap.Remove(NetId);
 		UserNumToNetIdMap.Remove(LocalUserNum);
 		UserNumToAccountIdMap.Remove(LocalUserNum);
@@ -1425,16 +1405,11 @@ ELoginStatus::Type FUserManagerEOS::GetLoginStatus(int32 LocalUserNum) const
 
 ELoginStatus::Type FUserManagerEOS::GetLoginStatus(const FUniqueNetIdEOS& UserId) const
 {
-	if (!StringToAccountIdMap.Contains(UserId.UniqueNetIdStr))
-	{
-		return ELoginStatus::NotLoggedIn;
-	}
-
 	FEOSSettings Settings = UEOSSettings::GetSettings();
 	// If the user isn't using EAS, then only check for a product user id
 	if (!Settings.bUseEAS)
 	{
-		EOS_ProductUserId ProductUserId = StringToProductUserIdMap[UserId.UniqueNetIdStr];
+		const EOS_ProductUserId ProductUserId = UserId.GetProductUserId();
 		if (ProductUserId != nullptr)
 		{
 			return ELoginStatus::LoggedIn;
@@ -1442,7 +1417,7 @@ ELoginStatus::Type FUserManagerEOS::GetLoginStatus(const FUniqueNetIdEOS& UserId
 		return ELoginStatus::NotLoggedIn;
 	}
 
-	EOS_EpicAccountId AccountId = StringToAccountIdMap[UserId.UniqueNetIdStr];
+	const EOS_EpicAccountId AccountId = UserId.GetEpicAccountId();
 	if (AccountId == nullptr)
 	{
 		return ELoginStatus::NotLoggedIn;
@@ -1773,7 +1748,6 @@ void FUserManagerEOS::AddRemotePlayer(int32 LocalUserNum, const FString& NetId,
 	NetIdStringToAttributeAccessMap.Emplace(NetId, AttributeRef);
 	EpicAccountIdToAttributeAccessMap.Emplace(EpicAccountId, AttributeRef);
 
-	StringToAccountIdMap.Emplace(NetId, EpicAccountId);
 	AccountIdToStringMap.Emplace(EpicAccountId, NetId);
 
 	// Read the user info for this player
@@ -1816,9 +1790,6 @@ void FUserManagerEOS::UpdateRemotePlayerProductUserId(EOS_EpicAccountId AccountI
 	AccountIdToStringMap.Emplace(AccountId, NewNetIdStr);
 	ProductUserIdToStringMap.Remove(UserId);
 	ProductUserIdToStringMap.Emplace(UserId, *NewNetIdStr);
-	StringToAccountIdMap.Remove(PrevNetIdStr);
-	StringToAccountIdMap.Emplace(NewNetIdStr, AccountId);
-	StringToProductUserIdMap.Emplace(NewNetIdStr, UserId);
 	// If it's the first time we're updating this remote player, it will be indexed by its epic account id
 	FOnlineUserPtr OnlineUser = NetIdStringToOnlineUserMap[AccountIdStr];
 	if (OnlineUser.IsValid())
@@ -1997,7 +1968,8 @@ bool FUserManagerEOS::SendInvite(int32 LocalUserNum, const FUniqueNetId& FriendI
 	}
 
 	const FUniqueNetIdEOS& EOSID = FUniqueNetIdEOS::Cast(FriendId);
-	if (!StringToAccountIdMap.Contains(EOSID.UniqueNetIdStr))
+	const EOS_EpicAccountId AccountId = EOSID.GetEpicAccountId();
+	if (EOS_EpicAccountId_IsValid(AccountId) == EOS_FALSE)
 	{
 		UE_LOG_ONLINE_FRIEND(Warning, TEXT("Can't SendInvite() for user (%d) since the potential player id is unknown"), LocalUserNum);
 		Delegate.ExecuteIfBound(LocalUserNum, false, FriendId, ListName, FString(TEXT("Can't SendInvite() for user (%d) since the player id is unknown"), LocalUserNum));
@@ -2021,7 +1993,7 @@ bool FUserManagerEOS::SendInvite(int32 LocalUserNum, const FUniqueNetId& FriendI
 	EOS_Friends_SendInviteOptions Options = { };
 	Options.ApiVersion = EOS_FRIENDS_SENDINVITE_API_LATEST;
 	Options.LocalUserId = UserNumToAccountIdMap[LocalUserNum];
-	Options.TargetUserId = StringToAccountIdMap[EOSID.UniqueNetIdStr];
+	Options.TargetUserId = AccountId;
 	EOS_Friends_SendInvite(EOSSubsystem->FriendsHandle, &Options, CallbackObj, CallbackObj->GetCallbackPtr());
 
 	return true;
@@ -2039,7 +2011,8 @@ bool FUserManagerEOS::AcceptInvite(int32 LocalUserNum, const FUniqueNetId& Frien
 	}
 
 	const FUniqueNetIdEOS& EOSID = FUniqueNetIdEOS::Cast(FriendId);
-	if (!StringToAccountIdMap.Contains(EOSID.UniqueNetIdStr))
+	const EOS_EpicAccountId AccountId = EOSID.GetEpicAccountId();
+	if (EOS_EpicAccountId_IsValid(AccountId) == EOS_FALSE)
 	{
 		UE_LOG_ONLINE_FRIEND(Warning, TEXT("Can't AcceptInvite() for user (%d) since the friend is not in their list"), LocalUserNum);
 		Delegate.ExecuteIfBound(LocalUserNum, false, FriendId, ListName, FString(TEXT("Can't AcceptInvite() for user (%d) since the friend is not in their list"), LocalUserNum));
@@ -2063,7 +2036,7 @@ bool FUserManagerEOS::AcceptInvite(int32 LocalUserNum, const FUniqueNetId& Frien
 	EOS_Friends_AcceptInviteOptions Options = { };
 	Options.ApiVersion = EOS_FRIENDS_ACCEPTINVITE_API_LATEST;
 	Options.LocalUserId = UserNumToAccountIdMap[LocalUserNum];
-	Options.TargetUserId = StringToAccountIdMap[EOSID.UniqueNetIdStr];
+	Options.TargetUserId = AccountId;
 	EOS_Friends_AcceptInvite(EOSSubsystem->FriendsHandle, &Options, CallbackObj, CallbackObj->GetCallbackPtr());
 	return true;
 }
@@ -2082,7 +2055,8 @@ bool FUserManagerEOS::RejectInvite(int32 LocalUserNum, const FUniqueNetId& Frien
 	}
 
 	const FUniqueNetIdEOS& EOSID = FUniqueNetIdEOS::Cast(FriendId);
-	if (!StringToAccountIdMap.Contains(EOSID.UniqueNetIdStr))
+	const EOS_EpicAccountId AccountId = EOSID.GetEpicAccountId();
+	if (EOS_EpicAccountId_IsValid(AccountId) == EOS_FALSE)
 	{
 		UE_LOG_ONLINE_FRIEND(Warning, TEXT("Can't RejectInvite() for user (%d) since the friend is not in their list"), LocalUserNum);
 		return false;
@@ -2091,7 +2065,7 @@ bool FUserManagerEOS::RejectInvite(int32 LocalUserNum, const FUniqueNetId& Frien
 	EOS_Friends_RejectInviteOptions Options{ 0 };
 	Options.ApiVersion = EOS_FRIENDS_REJECTINVITE_API_LATEST;
 	Options.LocalUserId = UserNumToAccountIdMap[LocalUserNum];
-	Options.TargetUserId = StringToAccountIdMap[EOSID.UniqueNetIdStr];
+	Options.TargetUserId = AccountId;
 	EOS_Friends_RejectInvite(EOSSubsystem->FriendsHandle, &Options, nullptr, &EOSRejectInviteCallback);
 	return true;
 }
@@ -2385,7 +2359,8 @@ typedef TEOSCallback<EOS_Presence_SetPresenceCompleteCallback, EOS_Presence_SetP
 void FUserManagerEOS::SetPresence(const FUniqueNetId& UserId, const FOnlineUserPresenceStatus& Status, const FOnPresenceTaskCompleteDelegate& Delegate)
 {
 	const FUniqueNetIdEOS& EOSID = FUniqueNetIdEOS::Cast(UserId);
-	if (!StringToAccountIdMap.Contains(EOSID.UniqueNetIdStr))
+	const EOS_EpicAccountId AccountId = EOSID.GetEpicAccountId();
+	if (EOS_EpicAccountId_IsValid(AccountId) == EOS_FALSE)
 	{
 		UE_LOG_ONLINE(Error, TEXT("Can't SetPresence() for user (%s) since they are not logged in"), *EOSID.UniqueNetIdStr);
 		return;
@@ -2394,7 +2369,7 @@ void FUserManagerEOS::SetPresence(const FUniqueNetId& UserId, const FOnlineUserP
 	EOS_HPresenceModification ChangeHandle = nullptr;
 	EOS_Presence_CreatePresenceModificationOptions Options = { };
 	Options.ApiVersion = EOS_PRESENCE_CREATEPRESENCEMODIFICATION_API_LATEST;
-	Options.LocalUserId = StringToAccountIdMap[EOSID.UniqueNetIdStr];
+	Options.LocalUserId = AccountId;
 	EOS_Presence_CreatePresenceModification(EOSSubsystem->PresenceHandle, &Options, &ChangeHandle);
 	if (ChangeHandle == nullptr)
 	{
@@ -2458,7 +2433,7 @@ void FUserManagerEOS::SetPresence(const FUniqueNetId& UserId, const FOnlineUserP
 
 	EOS_Presence_SetPresenceOptions PresOptions = { };
 	PresOptions.ApiVersion = EOS_PRESENCE_SETPRESENCE_API_LATEST;
-	PresOptions.LocalUserId = StringToAccountIdMap[EOSID.UniqueNetIdStr];
+	PresOptions.LocalUserId = AccountId;
 	PresOptions.PresenceModificationHandle = ChangeHandle;
 	// Last step commit the changes
 	EOS_Presence_SetPresence(EOSSubsystem->PresenceHandle, &PresOptions, CallbackObj, CallbackObj->GetCallbackPtr());
@@ -2477,10 +2452,10 @@ void FUserManagerEOS::QueryPresence(const FUniqueNetId& UserId, const FOnPresenc
 	}
 
 	const FUniqueNetIdEOS& EOSID = FUniqueNetIdEOS::Cast(UserId);
-	const FString& NetId = EOSID.UniqueNetIdStr;
-	if (!StringToAccountIdMap.Contains(NetId))
+	const EOS_EpicAccountId AccountId = EOSID.GetEpicAccountId();
+	if (EOS_EpicAccountId_IsValid(AccountId) == EOS_FALSE)
 	{
-		UE_LOG_ONLINE(Error, TEXT("Can't QueryPresence(%s) for unknown unique net id"), *NetId);
+		UE_LOG_ONLINE(Error, TEXT("Can't QueryPresence(%s) for unknown unique net id"), *EOSID.UniqueNetIdStr);
 		Delegate.ExecuteIfBound(UserId, false);
 		return;
 	}
@@ -2488,7 +2463,7 @@ void FUserManagerEOS::QueryPresence(const FUniqueNetId& UserId, const FOnPresenc
 	EOS_Presence_HasPresenceOptions HasOptions = { };
 	HasOptions.ApiVersion = EOS_PRESENCE_HASPRESENCE_API_LATEST;
 	HasOptions.LocalUserId = UserNumToAccountIdMap[DefaultLocalUser];
-	HasOptions.TargetUserId = StringToAccountIdMap[NetId];
+	HasOptions.TargetUserId = AccountId;
 	EOS_Bool bHasPresence = EOS_Presence_HasPresence(EOSSubsystem->PresenceHandle, &HasOptions);
 	if (bHasPresence == EOS_FALSE)
 	{
@@ -2617,19 +2592,19 @@ bool FUserManagerEOS::QueryUserInfo(int32 LocalUserNum, const TArray<FUniqueNetI
 			continue;
 		}
 		// Check to see if we know about this user or not
-		if (StringToAccountIdMap.Contains(EOSID.UniqueNetIdStr))
-		{
-			EOS_EpicAccountId AccountId = StringToAccountIdMap[EOSID.UniqueNetIdStr];
-			ReadUserInfo(LocalUserNum, AccountId);
-		}
-		else
+		const EOS_EpicAccountId AccountId = EOSID.GetEpicAccountId();
+		if (EOS_EpicAccountId_IsValid(AccountId) == EOS_TRUE)
 		{
-			// We need to build this one from the string
-			EOS_EpicAccountId AccountId = EOS_EpicAccountId_FromString(TCHAR_TO_UTF8(*EOSID.EpicAccountIdStr));
-			if (EOS_EpicAccountId_IsValid(AccountId) == EOS_TRUE)
+			// If the user is already registered, we'll update their user info
+			if (EpicAccountIdToAttributeAccessMap.Contains(AccountId))
 			{
-				UserEasIdsNeedingExternalMappings.Add(EOSID.EpicAccountIdStr);
-				
+				ReadUserInfo(LocalUserNum, AccountId);
+			}
+			else
+			{
+				// If the user is not registered, we'll add it and query their user info
+				UserEasIdsNeedingExternalMappings.Add(LexToString(AccountId));
+
 				// Registering the player will also query the user info data
 				AddRemotePlayer(LocalUserNum, EOSID.UniqueNetIdStr, AccountId);
 			}
@@ -2718,8 +2693,8 @@ typedef TEOSCallback<EOS_UserInfo_OnQueryUserInfoByDisplayNameCallback, EOS_User
 bool FUserManagerEOS::QueryUserIdMapping(const FUniqueNetId& UserId, const FString& DisplayNameOrEmail, const FOnQueryUserMappingComplete& Delegate)
 {
 	const FUniqueNetIdEOS& EOSID = FUniqueNetIdEOS::Cast(UserId);
-	const FString& NetId = EOSID.UniqueNetIdStr;
-	if (!StringToAccountIdMap.Contains(NetId))
+	const EOS_EpicAccountId AccountId = EOSID.GetEpicAccountId();
+	if (EOS_EpicAccountId_IsValid(AccountId) == EOS_FALSE)
 	{
 		UE_LOG_ONLINE(Error, TEXT("Specified local user (%s) is not known"), *EOSID.UniqueNetIdStr);
 		Delegate.ExecuteIfBound(false, UserId, DisplayNameOrEmail, *FUniqueNetIdEOS::EmptyId(), FString::Printf(TEXT("Specified local user (%s) is not known"), *EOSID.UniqueNetIdStr));
@@ -2760,7 +2735,7 @@ bool FUserManagerEOS::QueryUserIdMapping(const FUniqueNetId& UserId, const FStri
 
 	FQueryByDisplayNameOptions Options;
 	FCStringAnsi::Strncpy(Options.DisplayNameAnsi, TCHAR_TO_UTF8(*DisplayNameOrEmail), EOS_OSS_STRING_BUFFER_LENGTH);
-	Options.LocalUserId = StringToAccountIdMap[NetId];
+	Options.LocalUserId = AccountId;
 	EOS_UserInfo_QueryUserInfoByDisplayName(EOSSubsystem->UserInfoHandle, &Options, CallbackObj, CallbackObj->GetCallbackPtr());
 
 	return true;
@@ -2812,7 +2787,8 @@ typedef TEOSCallback<EOS_Connect_OnQueryExternalAccountMappingsCallback, EOS_Con
 bool FUserManagerEOS::QueryExternalIdMappings(const FUniqueNetId& UserId, const FExternalIdQueryOptions& QueryOptions, const TArray<FString>& ExternalIds, const FOnQueryExternalIdMappingsComplete& Delegate)
 {
 	const FUniqueNetIdEOS& EOSID = FUniqueNetIdEOS::Cast(UserId);
-	if (!StringToProductUserIdMap.Contains(EOSID.UniqueNetIdStr))
+	const EOS_EpicAccountId AccountId = EOSID.GetEpicAccountId();
+	if (EOS_EpicAccountId_IsValid(AccountId) == EOS_FALSE)
 	{
 		Delegate.ExecuteIfBound(false, UserId, QueryOptions, ExternalIds, FString::Printf(TEXT("User (%s) is not logged in, so can't query external account ids"), *EOSID.UniqueNetIdStr));
 		return false;
@@ -2822,7 +2798,7 @@ bool FUserManagerEOS::QueryExternalIdMappings(const FUniqueNetId& UserId, const
 	// Mark the queries as in progress
 	IsPlayerQueryExternalMappingsOngoingForLocalUserMap.FindOrAdd(LocalUserNum).Append(ExternalIds);
 
-	EOS_ProductUserId LocalUserId = StringToProductUserIdMap[EOSID.UniqueNetIdStr];
+	const EOS_ProductUserId LocalUserId = EOSID.GetProductUserId();
 	const int32 NumBatches = (ExternalIds.Num() / EOS_CONNECT_QUERYEXTERNALACCOUNTMAPPINGS_MAX_ACCOUNT_IDS) + 1;
 	int32 QueryStart = 0;
 	// Process queries in batches since there's a max that can be done at once
diff --git a/Engine/Plugins/Online/OnlineSubsystemEOS/Source/OnlineSubsystemEOS/Private/UserManagerEOS.h b/Engine/Plugins/Online/OnlineSubsystemEOS/Source/OnlineSubsystemEOS/Private/UserManagerEOS.h
--- a/Engine/Plugins/Online/OnlineSubsystemEOS/Source/OnlineSubsystemEOS/Private/UserManagerEOS.h
+++ b/Engine/Plugins/Online/OnlineSubsystemEOS/Source/OnlineSubsystemEOS/Private/UserManagerEOS.h
@@ -354,9 +354,6 @@ public:
 	int32 GetLocalUserNumFromUniqueNetId(const FUniqueNetId& NetId) const;
 	bool IsLocalUser(const FUniqueNetId& NetId) const;
 
-	EOS_EpicAccountId GetEpicAccountId(const FUniqueNetId& NetId) const;
-	EOS_ProductUserId GetProductUserId(const FUniqueNetId& NetId) const;
-
 	typedef TFunction<void(const EOS_ProductUserId& ProductUserId, EOS_EpicAccountId& EpicAccountId)> GetEpicAccountIdAsyncCallback;
 	bool GetEpicAccountIdFromProductUserId(const EOS_ProductUserId& ProductUserId, EOS_EpicAccountId& OutEpicAccountId) const;
 	void GetEpicAccountIdAsync(const EOS_ProductUserId& ProductUserId, const GetEpicAccountIdAsyncCallback& Callback) const;
@@ -430,8 +427,6 @@ private:
 	/** General account mappings */
 	TMap<EOS_EpicAccountId, FString> AccountIdToStringMap;
 	TMap<EOS_ProductUserId, FString> ProductUserIdToStringMap;
-	TMap<FString, EOS_EpicAccountId> StringToAccountIdMap;
-	TMap<FString, EOS_ProductUserId> StringToProductUserIdMap;
 
 	/** Per user friends lists accessible by user num or net id */
 	TMap<int32, FFriendsListEOSRef> LocalUserNumToFriendsListMap;
