Additon of ensure to FUniqueNetIdEOS FString constructor to enforce consistent formatting.
Addition of constructor taking id types to bypass MakeNetIdStringFromIds. Refactored use in affected callsites.

diff --git a/Engine/Plugins/Online/OnlineSubsystemEOS/Source/OnlineSubsystemEOS/Private/OnlineLeaderboardsEOS.cpp b/Engine/Plugins/Online/OnlineSubsystemEOS/Source/OnlineSubsystemEOS/Private/OnlineLeaderboardsEOS.cpp
--- a/Engine/Plugins/Online/OnlineSubsystemEOS/Source/OnlineSubsystemEOS/Private/OnlineLeaderboardsEOS.cpp
+++ b/Engine/Plugins/Online/OnlineSubsystemEOS/Source/OnlineSubsystemEOS/Private/OnlineLeaderboardsEOS.cpp
@@ -309,7 +309,7 @@ bool FOnlineLeaderboardsEOS::ReadLeaderboardsAroundRank(int32 Rank, uint32 Range
 						Nickname = TEXT("Unknown Player");
 					}
 					const EOS_EpicAccountId AccountId = EOS_EpicAccountId_FromString(EpicIdStr);
-					const FUniqueNetIdEOSRef NetId = FUniqueNetIdEOS::Create(MakeNetIdStringFromIds(AccountId, Record->UserId));
+					const FUniqueNetIdEOSRef NetId = FUniqueNetIdEOS::Create(AccountId, Record->UserId);
 
 					FOnlineStatsRow* Row = new(LambdaReadObject->Rows) FOnlineStatsRow(Nickname, NetId);
 					Row->Rank = Record->Rank;
diff --git a/Engine/Plugins/Online/OnlineSubsystemEOS/Source/OnlineSubsystemEOS/Private/OnlineSessionEOS.cpp b/Engine/Plugins/Online/OnlineSubsystemEOS/Source/OnlineSubsystemEOS/Private/OnlineSessionEOS.cpp
--- a/Engine/Plugins/Online/OnlineSubsystemEOS/Source/OnlineSubsystemEOS/Private/OnlineSessionEOS.cpp
+++ b/Engine/Plugins/Online/OnlineSubsystemEOS/Source/OnlineSubsystemEOS/Private/OnlineSessionEOS.cpp
@@ -21,6 +21,9 @@
 	#include "eos_metrics.h"
 	#include "eos_lobby.h"
 
+static FName EOS_SESSION_ID = TEXT("EOS_SESSION_ID");
+static FName EOS_LOBBY_ID = TEXT("EOS_LOBBY_ID");
+
 /** This is the game name plus version in ansi done once for optimization */
 char BucketIdAnsi[EOS_OSS_STRING_BUFFER_LENGTH];
 
@@ -348,19 +351,13 @@ struct FLobbyAttributeOptions :
 
 FOnlineSessionInfoEOS::FOnlineSessionInfoEOS()
 	: HostAddr(nullptr)
-	, SessionId(FUniqueNetIdEOS::EmptyId())
+	, SessionId(FUniqueNetIdString::EmptyId())
 	, SessionHandle(nullptr)
 	, bIsFromClone(false)
 {
 }
 
-FOnlineSessionInfoEOS::FOnlineSessionInfoEOS(const FString& InHostIp, const FString& InSessionId, EOS_HSessionDetails InSessionHandle)
-	: FOnlineSessionInfoEOS(InHostIp, FUniqueNetIdEOS::Create(InSessionId), InSessionHandle)
-{
-
-}
-
-FOnlineSessionInfoEOS::FOnlineSessionInfoEOS(const FString& InHostIp, FUniqueNetIdEOSRef UniqueNetId, EOS_HSessionDetails InSessionHandle)
+FOnlineSessionInfoEOS::FOnlineSessionInfoEOS(const FString& InHostIp, FUniqueNetIdStringRef UniqueNetId, EOS_HSessionDetails InSessionHandle)
 	: FOnlineSessionInfo()
 	, SessionId(UniqueNetId)
 	, SessionHandle(InSessionHandle)
@@ -510,7 +507,7 @@ void FOnlineSessionEOS::Init(const FString& InBucketId)
  *
  * @return pointer to the struct if found, NULL otherwise
  */
-FNamedOnlineSession* FOnlineSessionEOS::GetNamedSessionFromLobbyId(const FUniqueNetIdEOS& LobbyId)
+FNamedOnlineSession* FOnlineSessionEOS::GetNamedSessionFromLobbyId(const FUniqueNetIdString& LobbyId)
 {
 	FNamedOnlineSession* Result = nullptr;
 
@@ -604,7 +601,7 @@ void FOnlineSessionEOS::RegisterLobbyNotifications()
 
 void FOnlineSessionEOS::OnLobbyUpdateReceived(const EOS_LobbyId& LobbyId)
 {
-	const FUniqueNetIdEOSRef LobbyNetId = FUniqueNetIdEOS::Create(LobbyId);
+	const FUniqueNetIdStringRef LobbyNetId = FUniqueNetIdString::Create(UTF8_TO_TCHAR(LobbyId), EOS_LOBBY_ID);
 	FNamedOnlineSession* Session = GetNamedSessionFromLobbyId(*LobbyNetId);
 	if (Session)
 	{
@@ -651,16 +648,16 @@ void FOnlineSessionEOS::OnLobbyUpdateReceived(const EOS_LobbyId& LobbyId)
 
 void FOnlineSessionEOS::OnLobbyMemberUpdateReceived(const EOS_LobbyId& LobbyId, const EOS_ProductUserId& TargetUserId)
 {
-	const FUniqueNetIdEOSRef LobbyNetId = FUniqueNetIdEOS::Create(UTF8_TO_TCHAR(LobbyId));
+	const FUniqueNetIdStringRef LobbyNetId = FUniqueNetIdString::Create(UTF8_TO_TCHAR(LobbyId), EOS_LOBBY_ID);
 
 	EOSSubsystem->UserManager->GetEpicAccountIdAsync(TargetUserId, [this, LobbyNetId](const EOS_ProductUserId& ProductUserId, EOS_EpicAccountId& EpicAccountId)
 		{
 			FNamedOnlineSession* Session = GetNamedSessionFromLobbyId(*LobbyNetId);
 			if (Session)
 			{
-				FUniqueNetIdPtr UniqueNetId = EOSSubsystem->UserManager->CreateUniquePlayerId(MakeNetIdStringFromIds(EpicAccountId, ProductUserId));
+				FUniqueNetIdRef UniqueNetId = FUniqueNetIdEOS::Create(EpicAccountId, ProductUserId);
 
-				if (Session->SessionSettings.MemberSettings.Contains(UniqueNetId->AsShared()))
+				if (Session->SessionSettings.MemberSettings.Contains(UniqueNetId))
 				{
 					const FTCHARToUTF8 Utf8LobbyId(*LobbyNetId->ToString());
 
@@ -674,7 +671,7 @@ void FOnlineSessionEOS::OnLobbyMemberUpdateReceived(const EOS_LobbyId& LobbyId,
 					EOS_EResult Result = EOS_Lobby_CopyLobbyDetailsHandle(LobbyHandle, &Options, &LobbyDetailsHandle);
 					if (Result == EOS_EResult::EOS_Success)
 					{
-						CopyLobbyMemberAttributes(LobbyDetailsHandle, ProductUserId, Session->SessionSettings.MemberSettings[UniqueNetId->AsShared()]);
+						CopyLobbyMemberAttributes(LobbyDetailsHandle, ProductUserId, Session->SessionSettings.MemberSettings[UniqueNetId]);
 
 						TriggerOnSessionParticipantSettingsUpdatedDelegates(Session->SessionName, *UniqueNetId, Session->SessionSettings);
 
@@ -699,7 +696,7 @@ void FOnlineSessionEOS::OnLobbyMemberUpdateReceived(const EOS_LobbyId& LobbyId,
 
 void FOnlineSessionEOS::OnMemberStatusReceived(const EOS_LobbyId& LobbyId, const EOS_ProductUserId& TargetUserId, EOS_ELobbyMemberStatus CurrentStatus)
 {
-	const FUniqueNetIdEOSRef LobbyNetId = FUniqueNetIdEOS::Create(UTF8_TO_TCHAR(LobbyId));
+	const FUniqueNetIdStringRef LobbyNetId = FUniqueNetIdString::Create(UTF8_TO_TCHAR(LobbyId), EOS_LOBBY_ID);
 	FNamedOnlineSession* Session = GetNamedSessionFromLobbyId(*LobbyNetId);
 	if (Session)
 	{
@@ -718,7 +715,7 @@ void FOnlineSessionEOS::OnMemberStatusReceived(const EOS_LobbyId& LobbyId, const
 						FNamedOnlineSession* Session = GetNamedSessionFromLobbyId(*LobbyNetId);
 						if (Session)
 						{
-							FUniqueNetIdPtr UniqueNetId = EOSSubsystem->UserManager->CreateUniquePlayerId(MakeNetIdStringFromIds(EpicAccountId, ProductUserId));
+							FUniqueNetIdRef UniqueNetId = FUniqueNetIdEOS::Create(EpicAccountId, ProductUserId);
 
 							UnregisterPlayer(Session->SessionName, *UniqueNetId);
 
@@ -744,7 +741,7 @@ void FOnlineSessionEOS::OnMemberStatusReceived(const EOS_LobbyId& LobbyId, const
 						{
 							int32 DefaultLocalUser = EOSSubsystem->UserManager->GetDefaultLocalUser();
 
-							FUniqueNetIdPtr RemovedUserNetId = EOSSubsystem->UserManager->CreateUniquePlayerId(MakeNetIdStringFromIds(EpicAccountId, ProductUserId));
+							FUniqueNetIdRef RemovedUserNetId = FUniqueNetIdEOS::Create(EpicAccountId, ProductUserId);
 							FUniqueNetIdPtr LocalPlayerNetId = EOSSubsystem->UserManager->GetUniquePlayerId(DefaultLocalUser);
 
 							if (*LocalPlayerNetId == *RemovedUserNetId)
@@ -768,7 +765,7 @@ void FOnlineSessionEOS::OnMemberStatusReceived(const EOS_LobbyId& LobbyId, const
 						{
 							int32 DefaultLocalUser = EOSSubsystem->UserManager->GetDefaultLocalUser();
 
-							FUniqueNetIdPtr NewOwnerUniqueNetId = EOSSubsystem->UserManager->CreateUniquePlayerId(MakeNetIdStringFromIds(EpicAccountId, ProductUserId));
+							FUniqueNetIdRef NewOwnerUniqueNetId = FUniqueNetIdEOS::Create(EpicAccountId, ProductUserId);
 							FUniqueNetIdPtr LocalPlayerUniqueNetId = EOSSubsystem->UserManager->GetUniquePlayerId(DefaultLocalUser);
 
 							if (*LocalPlayerUniqueNetId == *NewOwnerUniqueNetId)
@@ -1213,7 +1210,7 @@ uint32 FOnlineSessionEOS::CreateEOSSession(int32 HostingPlayerNum, FNamedOnlineS
 		// This is basically ignored
 		HostAddr = TEXT("127.0.0.1");
 	}
-	Session->SessionInfo = MakeShareable(new FOnlineSessionInfoEOS(HostAddr, FUniqueNetIdEOS::EmptyId(), nullptr));
+	Session->SessionInfo = MakeShareable(new FOnlineSessionInfoEOS(HostAddr, FUniqueNetIdString::Create(FString(), EOS_SESSION_ID), nullptr));
 
 	FName SessionName = Session->SessionName;
 
@@ -1231,7 +1228,7 @@ uint32 FOnlineSessionEOS::CreateEOSSession(int32 HostingPlayerNum, FNamedOnlineS
 				TSharedPtr<FOnlineSessionInfoEOS> SessionInfo = StaticCastSharedPtr<FOnlineSessionInfoEOS>(Session->SessionInfo);
 				if (SessionInfo.IsValid())
 				{
-					SessionInfo->SessionId = FUniqueNetIdEOS::Create(Data->SessionId);
+					SessionInfo->SessionId = FUniqueNetIdString::Create(UTF8_TO_TCHAR(Data->SessionId), EOS_SESSION_ID);
 				}
 
 				Session->SessionState = EOnlineSessionState::Pending;
@@ -1832,7 +1829,7 @@ bool FOnlineSessionEOS::FindSessionById(const FUniqueNetId& SearchingUserId, con
 		CurrentSessionSearch->SearchState = EOnlineAsyncTaskState::InProgress;
 
 		StartLobbySearch(EOSSubsystem->UserManager->GetLocalUserNumFromUniqueNetId(SearchingUserId), LobbySearchHandle, CurrentSessionSearch.ToSharedRef(),
-			FOnSingleSessionResultCompleteDelegate::CreateLambda([this, OrigCallback = FOnSingleSessionResultCompleteDelegate(CompletionDelegate), SessId = CreateSessionIdFromString(SessionId.ToString())](int32 LocalUserNum, bool bWasSuccessful, const FOnlineSessionSearchResult& EOSResult)
+			FOnSingleSessionResultCompleteDelegate::CreateLambda([this, OrigCallback = FOnSingleSessionResultCompleteDelegate(CompletionDelegate), SessId = FUniqueNetIdString::Create(SessionId.ToString(), EOS_SESSION_ID)](int32 LocalUserNum, bool bWasSuccessful, const FOnlineSessionSearchResult& EOSResult)
 		{
 			if (bWasSuccessful)
 			{
@@ -2003,7 +2000,7 @@ void FOnlineSessionEOS::AddSearchResult(EOS_HSessionDetails SessionHandle, const
 		int32 Position = SearchSettings->SearchResults.AddZeroed();
 		FOnlineSessionSearchResult& SearchResult = SearchSettings->SearchResults[Position];
 		// This will set the host address and port
-		SearchResult.Session.SessionInfo = MakeShareable(new FOnlineSessionInfoEOS(SessionInfo->HostAddress, SessionInfo->SessionId, SessionHandle));
+		SearchResult.Session.SessionInfo = MakeShareable(new FOnlineSessionInfoEOS(SessionInfo->HostAddress, FUniqueNetIdString::Create(SessionInfo->SessionId, EOS_SESSION_ID), SessionHandle));
 
 		CopySearchResult(SessionHandle, SessionInfo, SearchResult.Session);
 
@@ -3279,12 +3276,7 @@ bool FOnlineSessionEOS::IsHost(const FNamedOnlineSession& Session) const
 
 FUniqueNetIdPtr FOnlineSessionEOS::CreateSessionIdFromString(const FString& SessionIdStr)
 {
-	FUniqueNetIdPtr SessionId;
-	if (!SessionIdStr.IsEmpty())
-	{
-		SessionId = FUniqueNetIdEOS::Create(SessionIdStr);
-	}
-	return SessionId;
+	return FUniqueNetIdString::Create(SessionIdStr, EOS_SESSION_ID);
 }
 
 EOS_ELobbyPermissionLevel FOnlineSessionEOS::GetLobbyPermissionLevelFromSessionSettings(const FOnlineSessionSettings& SessionSettings)
@@ -3369,7 +3361,7 @@ uint32 FOnlineSessionEOS::CreateLobbySession(int32 HostingPlayerNum, FNamedOnlin
 				FInternetAddrEOS TempAddr(LexToString(LocalProductUserId), SessionName.ToString(), FURL::UrlConfig.DefaultPort);
 				FString HostAddr = TempAddr.ToString(true);
 
-				Session->SessionInfo = MakeShareable(new FOnlineSessionInfoEOS(HostAddr, Data->LobbyId, nullptr));
+				Session->SessionInfo = MakeShareable(new FOnlineSessionInfoEOS(HostAddr, FUniqueNetIdString::Create(Data->LobbyId, EOS_LOBBY_ID), nullptr));
 
 #if WITH_EOS_RTC
 				if (FEOSVoiceChatUser* VoiceChatUser = static_cast<FEOSVoiceChatUser*>(EOSSubsystem->GetEOSVoiceChatUserInterface(*LocalUserNetId)))
@@ -3446,7 +3438,7 @@ uint32 FOnlineSessionEOS::JoinLobbySession(int32 PlayerNum, FNamedOnlineSession*
 						BeginSessionAnalytics(Session);
 						
 						// Initialize the OSS's member list with the current member list of the lobby.
-						const FUniqueNetIdEOSRef LobbyNetId = FUniqueNetIdEOS::Create(UTF8_TO_TCHAR(Data->LobbyId));
+						const FUniqueNetIdStringRef LobbyNetId = FUniqueNetIdString::Create(UTF8_TO_TCHAR(Data->LobbyId), EOS_LOBBY_ID);
 						EOS_LobbyDetails_GetMemberCountOptions GetMemberCountOptions = { 0 };
 						GetMemberCountOptions.ApiVersion = EOS_LOBBYDETAILS_GETMEMBERCOUNT_API_LATEST;
 						const uint32 MemberCount = EOS_LobbyDetails_GetMemberCount(LobbyDetailsHandle, &GetMemberCountOptions);
@@ -3645,13 +3637,13 @@ void FOnlineSessionEOS::SetLobbyAttributes(EOS_HLobbyModification LobbyModificat
 	}
 }
 
-void FOnlineSessionEOS::AddLobbyMember(const FUniqueNetIdEOSRef LobbyNetId, const EOS_ProductUserId& TargetUserId)
+void FOnlineSessionEOS::AddLobbyMember(const FUniqueNetIdStringRef LobbyNetId, const EOS_ProductUserId& TargetUserId)
 {
 	EOSSubsystem->UserManager->GetEpicAccountIdAsync(TargetUserId, [this, LobbyNetId](const EOS_ProductUserId& ProductUserId, EOS_EpicAccountId& EpicAccountId)
 		{
 			if (const FNamedOnlineSession* Session = GetNamedSessionFromLobbyId(*LobbyNetId))
 			{
-				const FUniqueNetIdPtr UniqueNetId = EOSSubsystem->UserManager->CreateUniquePlayerId(MakeNetIdStringFromIds(EpicAccountId, ProductUserId));
+				const FUniqueNetIdRef UniqueNetId = FUniqueNetIdEOS::Create(EpicAccountId, ProductUserId);
 
 				RegisterPlayer(Session->SessionName, *UniqueNetId, false);
 
@@ -3960,7 +3952,7 @@ void FOnlineSessionEOS::AddLobbySearchResult(EOS_HLobbyDetails LobbyDetailsHandl
 		FInternetAddrEOS TempAddr(LexToString(LobbyDetailsInfo->LobbyOwnerUserId), NetDriverName.ToString(), GetTypeHash(NetDriverName.ToString()));
 		FString HostAddr = TempAddr.ToString(true);
 
-		SearchResult.Session.SessionInfo = MakeShareable(new FOnlineSessionInfoEOS(HostAddr, LobbyDetailsInfo->LobbyId, nullptr));
+		SearchResult.Session.SessionInfo = MakeShareable(new FOnlineSessionInfoEOS(HostAddr, FUniqueNetIdString::Create(LobbyDetailsInfo->LobbyId, EOS_LOBBY_ID), nullptr));
 
 		// We copy the lobby data and settings
 		CopyLobbyData(LobbyDetailsHandle, LobbyDetailsInfo, SearchResult.Session);
@@ -3980,15 +3972,15 @@ void FOnlineSessionEOS::AddLobbySearchResult(EOS_HLobbyDetails LobbyDetailsHandl
 
 			EOSSubsystem->UserManager->GetEpicAccountIdAsync(TargetUserId, [this, LobbyDetailsHandle, SearchSettings, Position](const EOS_ProductUserId& ProductUserId, EOS_EpicAccountId& EpicAccountId) 
 				{
-					FUniqueNetIdPtr UniqueNetId = EOSSubsystem->UserManager->CreateUniquePlayerId(MakeNetIdStringFromIds(EpicAccountId, ProductUserId));
+					FUniqueNetIdRef UniqueNetId = FUniqueNetIdEOS::Create(EpicAccountId, ProductUserId);
 
 					FOnlineSessionSearchResult& SearchResult = SearchSettings->SearchResults[Position];
-					if (!SearchResult.Session.SessionSettings.MemberSettings.Contains(UniqueNetId->AsShared()))
+					if (!SearchResult.Session.SessionSettings.MemberSettings.Contains(UniqueNetId))
 					{
-						SearchResult.Session.SessionSettings.MemberSettings.Add(UniqueNetId->AsShared(), FSessionSettings());
+						SearchResult.Session.SessionSettings.MemberSettings.Add(UniqueNetId, FSessionSettings());
 					}
 
-					CopyLobbyMemberAttributes(LobbyDetailsHandle, ProductUserId, SearchResult.Session.SessionSettings.MemberSettings[UniqueNetId->AsShared()]);
+					CopyLobbyMemberAttributes(LobbyDetailsHandle, ProductUserId, SearchResult.Session.SessionSettings.MemberSettings[UniqueNetId]);
 				});
 		}
 
diff --git a/Engine/Plugins/Online/OnlineSubsystemEOS/Source/OnlineSubsystemEOS/Private/OnlineSessionEOS.h b/Engine/Plugins/Online/OnlineSubsystemEOS/Source/OnlineSubsystemEOS/Private/OnlineSessionEOS.h
--- a/Engine/Plugins/Online/OnlineSubsystemEOS/Source/OnlineSubsystemEOS/Private/OnlineSessionEOS.h
+++ b/Engine/Plugins/Online/OnlineSubsystemEOS/Source/OnlineSubsystemEOS/Private/OnlineSessionEOS.h
@@ -227,7 +227,7 @@ private:
 	void SetLobbyAttributes(EOS_HLobbyModification LobbyModificationHandle, FNamedOnlineSession* Session);
 	void AddLobbyAttribute(EOS_HLobbyModification LobbyModificationHandle, const EOS_Lobby_AttributeData* Attribute);
 	void AddLobbyMemberAttribute(EOS_HLobbyModification LobbyModificationHandle, const EOS_Lobby_AttributeData* Attribute);
-	void AddLobbyMember(const FUniqueNetIdEOSRef LobbyNetId, const EOS_ProductUserId& TargetUserId);
+	void AddLobbyMember(const FUniqueNetIdStringRef LobbyNetId, const EOS_ProductUserId& TargetUserId);
 
 	// Lobby search
 	void AddLobbySearchAttribute(EOS_HLobbySearch LobbySearchHandle, const EOS_Lobby_AttributeData* Attribute, EOS_EOnlineComparisonOp ComparisonOp);
@@ -241,7 +241,7 @@ private:
 
 	void GetEpicAccountIdAsync(const EOS_ProductUserId& ProductUserId, const GetEpicAccountIdAsyncCallback& Callback);
 	void RegisterLobbyNotifications();
-	FNamedOnlineSession* GetNamedSessionFromLobbyId(const FUniqueNetIdEOS& LobbyId);
+	FNamedOnlineSession* GetNamedSessionFromLobbyId(const FUniqueNetIdString& LobbyId);
 	bool GetEpicAccountIdFromProductUserId(const EOS_ProductUserId& ProductUserId, EOS_EpicAccountId& EpicAccountId);
 	EOS_ELobbyPermissionLevel GetLobbyPermissionLevelFromSessionSettings(const FOnlineSessionSettings& SessionSettings);
 	uint32_t GetLobbyMaxMembersFromSessionSettings(const FOnlineSessionSettings& SessionSettings);
diff --git a/Engine/Plugins/Online/OnlineSubsystemEOS/Source/OnlineSubsystemEOS/Private/OnlineSubsystemEOSTypes.h b/Engine/Plugins/Online/OnlineSubsystemEOS/Source/OnlineSubsystemEOS/Private/OnlineSubsystemEOSTypes.h
--- a/Engine/Plugins/Online/OnlineSubsystemEOS/Source/OnlineSubsystemEOS/Private/OnlineSubsystemEOSTypes.h
+++ b/Engine/Plugins/Online/OnlineSubsystemEOS/Source/OnlineSubsystemEOS/Private/OnlineSubsystemEOSTypes.h
@@ -4,6 +4,7 @@
 
 #include "OnlineSubsystemTypes.h"
 
+#include "OnlineSubsystem.h"
 #include "Interfaces/OnlineIdentityInterface.h"
 #include "Interfaces/OnlineFriendsInterface.h"
 #include "Interfaces/OnlinePresenceInterface.h"
@@ -25,6 +26,11 @@ class FOnlineSubsystemEOS;
 typedef TSharedPtr<const class FUniqueNetIdEOS> FUniqueNetIdEOSPtr;
 typedef TSharedRef<const class FUniqueNetIdEOS> FUniqueNetIdEOSRef;
 
+static inline FString MakeNetIdStringFromIds(EOS_EpicAccountId AccountId, EOS_ProductUserId UserId)
+{
+	return LexToString(AccountId) + EOS_ID_SEPARATOR + LexToString(UserId);
+}
+
 /**
  * Unique net id wrapper for a EOS account ids. The underlying string is a combination
  * of both account ids concatenated. "<EOS_EpicAccountId>|<EOS_ProductAccountId>"
@@ -96,32 +102,39 @@ public:
 
 	void ParseAccountIds()
 	{
-		TArray<FString> AccountIds;
-		UniqueNetIdStr.ParseIntoArray(AccountIds, EOS_ID_SEPARATOR, false);
-
-		FString EpicAccountIdStr = EMPTY_EASID;
-		if (AccountIds.Num() > 0 && AccountIds[0].Len() > 0)
-		{
-			EpicAccountIdStr = AccountIds[0];
-			EpicAccountId = EOS_EpicAccountId_FromString(TCHAR_TO_UTF8(*EpicAccountIdStr));
-		}
-		else
+		if (ensure(UniqueNetIdStr.Contains(EOS_ID_SEPARATOR)))
 		{
-			EpicAccountId = nullptr;
-		}
-		AddToBuffer(RawBytes, EpicAccountIdStr);
-
-		FString ProductUserIdStr = EMPTY_PUID;
-		if (AccountIds.Num() > 1 && AccountIds[1].Len() > 0)
-		{
-			ProductUserIdStr = AccountIds[1];
-			ProductUserId = EOS_ProductUserId_FromString(TCHAR_TO_UTF8(*ProductUserIdStr));
+			TArray<FString> AccountIds;
+			UniqueNetIdStr.ParseIntoArray(AccountIds, EOS_ID_SEPARATOR, false);
+
+			FString EpicAccountIdStr = EMPTY_EASID;
+			if (AccountIds.Num() > 0 && AccountIds[0].Len() > 0)
+			{
+				EpicAccountIdStr = AccountIds[0];
+				EpicAccountId = EOS_EpicAccountId_FromString(TCHAR_TO_UTF8(*EpicAccountIdStr));
+			}
+			else
+			{
+				EpicAccountId = nullptr;
+			}
+			AddToBuffer(RawBytes, EpicAccountIdStr);
+
+			FString ProductUserIdStr = EMPTY_PUID;
+			if (AccountIds.Num() > 1 && AccountIds[1].Len() > 0)
+			{
+				ProductUserIdStr = AccountIds[1];
+				ProductUserId = EOS_ProductUserId_FromString(TCHAR_TO_UTF8(*ProductUserIdStr));
+			}
+			else
+			{
+				ProductUserId = nullptr;
+			}
+			AddToBuffer(RawBytes + ID_HALF_BYTE_SIZE, ProductUserIdStr);
 		}
 		else
 		{
-			ProductUserId = nullptr;
+			UE_LOG_ONLINE(Warning, TEXT("Invalid format for String in FUniqueNetIdEOS constructor. Please use FUniqueNetIdEOS::Create(EOS_EpicAccountId, EOS_ProductUserId) or FUserManagerEOS::MakeNetIdStringFromIds"));
 		}
-		AddToBuffer(RawBytes + ID_HALF_BYTE_SIZE, ProductUserIdStr);
 	}
 
 	void AddToBuffer(uint8* Buffer, const FString& Source)
@@ -168,6 +181,14 @@ private:
 		Type = FName("EOS");
 	}
 
+	explicit FUniqueNetIdEOS(EOS_EpicAccountId InEpicAccountId, EOS_ProductUserId InProductUserId)
+		PRAGMA_DISABLE_DEPRECATION_WARNINGS
+		: FUniqueNetIdString(MakeNetIdStringFromIds(InEpicAccountId, InProductUserId), FName("EOS"))
+		PRAGMA_ENABLE_DEPRECATION_WARNINGS
+	{
+		ParseAccountIds();
+	}
+
 	explicit FUniqueNetIdEOS(const FString& InUniqueNetId)
 		PRAGMA_DISABLE_DEPRECATION_WARNINGS
 		: FUniqueNetIdString(InUniqueNetId, FName("EOS"))
@@ -442,22 +463,6 @@ protected:
 	FDateTime LastSeenTime;
 };
 
-static inline FString MakeNetIdStringFromIds(EOS_EpicAccountId AccountId, EOS_ProductUserId UserId)
-{
-	FString NetId = LexToString(AccountId);
-
-	// Only add this when the product user id is valid for more consistent net id string generation
-	// across different code paths
-	const FString ProductIdStr = LexToString(UserId);
-	if (!ProductIdStr.IsEmpty())
-	{
-		NetId += EOS_ID_SEPARATOR;
-		NetId += ProductIdStr;
-	}
-
-	return NetId;
-}
-
 /** Class to handle all callbacks generically using a lambda to process callback results */
 template<typename CallbackFuncType, typename CallbackType, typename OwningType>
 class TEOSCallback :
@@ -741,9 +746,7 @@ PACKAGE_SCOPE:
 	{
 	}
 
-	FOnlineSessionInfoEOS(const FString& InHostIp, const FString& InSessionId, EOS_HSessionDetails InSessionHandle);
-
-	FOnlineSessionInfoEOS(const FString& InHostIp, FUniqueNetIdEOSRef UniqueNetId, EOS_HSessionDetails InSessionHandle);
+	FOnlineSessionInfoEOS(const FString& InHostIp, FUniqueNetIdStringRef UniqueNetId, EOS_HSessionDetails InSessionHandle);
 
 	/**
 	 * Initialize LAN session
diff --git a/Engine/Plugins/Online/OnlineSubsystemEOS/Source/OnlineSubsystemEOS/Private/UserManagerEOS.cpp b/Engine/Plugins/Online/OnlineSubsystemEOS/Source/OnlineSubsystemEOS/Private/UserManagerEOS.cpp
--- a/Engine/Plugins/Online/OnlineSubsystemEOS/Source/OnlineSubsystemEOS/Private/UserManagerEOS.cpp
+++ b/Engine/Plugins/Online/OnlineSubsystemEOS/Source/OnlineSubsystemEOS/Private/UserManagerEOS.cpp
@@ -1708,7 +1708,7 @@ void FUserManagerEOS::FriendStatusChanged(const EOS_Friends_OnFriendsUpdateInfo*
 
 void FUserManagerEOS::AddFriend(int32 LocalUserNum, EOS_EpicAccountId EpicAccountId)
 {
-	const FString NetId = LexToString(EpicAccountId);
+	const FString NetId = MakeNetIdStringFromIds(EpicAccountId, nullptr);
 	FUniqueNetIdEOSRef FriendNetId = FUniqueNetIdEOS::Create(NetId);
 	FOnlineFriendEOSRef FriendRef = MakeShareable(new FOnlineFriendEOS(FriendNetId));
 	LocalUserNumToFriendsListMap[LocalUserNum]->Add(NetId, FriendRef);
@@ -1765,7 +1765,6 @@ void FUserManagerEOS::UpdateRemotePlayerProductUserId(EOS_EpicAccountId AccountI
 		return;
 	}
 
-	const FString AccountIdStr = LexToString(AccountId);
 	const FString UserIdStr = LexToString(UserId);
 
 	// Get the unique net id and rebuild the string for it
@@ -1790,17 +1789,8 @@ void FUserManagerEOS::UpdateRemotePlayerProductUserId(EOS_EpicAccountId AccountI
 	AccountIdToStringMap.Emplace(AccountId, NewNetIdStr);
 	ProductUserIdToStringMap.Remove(UserId);
 	ProductUserIdToStringMap.Emplace(UserId, *NewNetIdStr);
-	// If it's the first time we're updating this remote player, it will be indexed by its epic account id
-	FOnlineUserPtr OnlineUser = NetIdStringToOnlineUserMap[AccountIdStr];
-	if (OnlineUser.IsValid())
-	{
-		NetIdStringToOnlineUserMap.Remove(AccountIdStr);
-	}
-	else
-	{
-		OnlineUser = NetIdStringToOnlineUserMap[PrevNetIdStr];
-		NetIdStringToOnlineUserMap.Remove(PrevNetIdStr);
-	}
+	FOnlineUserPtr OnlineUser = NetIdStringToOnlineUserMap[PrevNetIdStr];
+	NetIdStringToOnlineUserMap.Remove(PrevNetIdStr);
 	NetIdStringToOnlineUserMap.Emplace(NewNetIdStr, OnlineUser);
 	NetIdStringToAttributeAccessMap.Remove(PrevNetIdStr);
 	NetIdStringToAttributeAccessMap.Emplace(NewNetIdStr, AttrAccess);
