Added logic to send and receive Owner Id and Name as attributes in EOS Sessions.

diff --git a/Engine/Plugins/Online/OnlineSubsystemEOS/Source/OnlineSubsystemEOS/Private/OnlineSessionEOS.cpp b/Engine/Plugins/Online/OnlineSubsystemEOS/Source/OnlineSubsystemEOS/Private/OnlineSessionEOS.cpp
--- a/Engine/Plugins/Online/OnlineSubsystemEOS/Source/OnlineSubsystemEOS/Private/OnlineSessionEOS.cpp
+++ b/Engine/Plugins/Online/OnlineSubsystemEOS/Source/OnlineSubsystemEOS/Private/OnlineSessionEOS.cpp
@@ -1049,6 +1049,12 @@ void FOnlineSessionEOS::SetAttributes(EOS_HSessionModification SessionModHandle,
 	FAttributeOptions Opt2("NumPublicConnections", Session->SessionSettings.NumPublicConnections);
 	AddAttribute(SessionModHandle, &Opt2);
 
+	if (Session->OwningUserId->IsValid())
+	{
+		FAttributeOptions OwningUserId("OwningUserId", TCHAR_TO_UTF8(*Session->OwningUserId->ToString()));
+		AddAttribute(SessionModHandle, &OwningUserId);
+	}
+
 	// Handle auto generation of dedicated server names
 	if (Session->OwningUserName.IsEmpty())
 	{
@@ -1067,6 +1073,9 @@ void FOnlineSessionEOS::SetAttributes(EOS_HSessionModification SessionModHandle,
 		Session->OwningUserName = OwningPlayerName;
 	}
 
+	FAttributeOptions OwningUserName("OwningUserName", TCHAR_TO_UTF8(*Session->OwningUserName));
+	AddAttribute(SessionModHandle, &OwningUserName);
+
 	FAttributeOptions Opt5("bAntiCheatProtected", Session->SessionSettings.bAntiCheatProtected);
 	AddAttribute(SessionModHandle, &Opt5);
 
@@ -1938,6 +1947,14 @@ void FOnlineSessionEOS::CopyAttributes(EOS_HSessionDetails SessionHandle, FOnlin
 				// Adjust the private connections based upon this
 				OutSession.SessionSettings.NumPrivateConnections = Attribute->Data->Value.AsInt64;
 			}
+			else if (Key == TEXT("OwningUserId"))
+			{
+				OutSession.OwningUserId = FUniqueNetIdEOS::Create(UTF8_TO_TCHAR(Attribute->Data->Value.AsUtf8));
+			}
+			else if (Key == TEXT("OwningUserName"))
+			{
+				OutSession.OwningUserName = UTF8_TO_TCHAR(Attribute->Data->Value.AsUtf8);
+			}
 			else if (Key == TEXT("bAntiCheatProtected"))
 			{
 				OutSession.SessionSettings.bAntiCheatProtected = Attribute->Data->Value.AsBool == EOS_TRUE;
