Added missing CREATE_MIRROR_PLATFORM_SESSION checks in OnlineSessionEOSPlus

diff --git a/Engine/Plugins/Online/OnlineSubsystemEOS/Source/OnlineSubsystemEOSPlus/Private/OnlineSessionEOSPlus.cpp b/Engine/Plugins/Online/OnlineSubsystemEOS/Source/OnlineSubsystemEOSPlus/Private/OnlineSessionEOSPlus.cpp
--- a/Engine/Plugins/Online/OnlineSubsystemEOS/Source/OnlineSubsystemEOSPlus/Private/OnlineSessionEOSPlus.cpp
+++ b/Engine/Plugins/Online/OnlineSubsystemEOS/Source/OnlineSubsystemEOSPlus/Private/OnlineSessionEOSPlus.cpp
@@ -430,7 +430,12 @@ bool FOnlineSessionEOSPlus::StartSession(FName SessionName)
 			return false;
 		}
 	}
+
+#if CREATE_MIRROR_PLATFORM_SESSION
 	return BaseSessionInterface->StartSession(SessionName);
+#else
+	return true;
+#endif
 }
 
 bool FOnlineSessionEOSPlus::UpdateSession(FName SessionName, FOnlineSessionSettings& UpdatedSessionSettings, bool bShouldRefreshOnlineData)
@@ -442,7 +447,12 @@ bool FOnlineSessionEOSPlus::UpdateSession(FName SessionName, FOnlineSessionSetti
 			return false;
 		}
 	}
+
+#if CREATE_MIRROR_PLATFORM_SESSION
 	return BaseSessionInterface->UpdateSession(SessionName, UpdatedSessionSettings, bShouldRefreshOnlineData);
+#else
+	return true;
+#endif
 }
 
 bool FOnlineSessionEOSPlus::EndSession(FName SessionName)
@@ -454,7 +464,12 @@ bool FOnlineSessionEOSPlus::EndSession(FName SessionName)
 			return false;
 		}
 	}
+
+#if CREATE_MIRROR_PLATFORM_SESSION
 	return BaseSessionInterface->EndSession(SessionName);
+#else
+	return true;
+#endif
 }
 
 FOnDestroySessionCompleteDelegate IgnoredDestroySessionDelegate;
@@ -468,7 +483,12 @@ bool FOnlineSessionEOSPlus::DestroySession(FName SessionName, const FOnDestroySe
 			return false;
 		}
 	}
+
+#if CREATE_MIRROR_PLATFORM_SESSION
 	return BaseSessionInterface->DestroySession(SessionName, IgnoredDestroySessionDelegate);
+#else
+	return true;
+#endif
 }
 
 bool FOnlineSessionEOSPlus::IsPlayerInSession(FName SessionName, const FUniqueNetId& UniqueId)
@@ -799,7 +819,12 @@ bool FOnlineSessionEOSPlus::RegisterPlayer(FName SessionName, const FUniqueNetId
 	{
 		return false;
 	}
+
+#if CREATE_MIRROR_PLATFORM_SESSION
 	return BaseSessionInterface->RegisterPlayer(SessionName, *Id, bWasInvited);
+#else
+	return true;
+#endif
 }
 
 bool FOnlineSessionEOSPlus::RegisterPlayers(FName SessionName, const TArray<FUniqueNetIdRef>& Players, bool bWasInvited)
@@ -811,7 +836,12 @@ bool FOnlineSessionEOSPlus::RegisterPlayers(FName SessionName, const TArray<FUni
 			return false;
 		}
 	}
+
+#if CREATE_MIRROR_PLATFORM_SESSION
 	return BaseSessionInterface->RegisterPlayers(SessionName, GetBaseNetIds(Players), bWasInvited);
+#else
+	return true;
+#endif
 }
 
 bool FOnlineSessionEOSPlus::UnregisterPlayer(FName SessionName, const FUniqueNetId& PlayerId)
@@ -833,7 +863,12 @@ bool FOnlineSessionEOSPlus::UnregisterPlayer(FName SessionName, const FUniqueNet
 	{
 		return false;
 	}
+
+#if CREATE_MIRROR_PLATFORM_SESSION
 	return BaseSessionInterface->UnregisterPlayer(SessionName, *Id);
+#else
+	return true;
+#endif
 }
 
 bool FOnlineSessionEOSPlus::UnregisterPlayers(FName SessionName, const TArray<FUniqueNetIdRef>& Players)
@@ -845,7 +880,12 @@ bool FOnlineSessionEOSPlus::UnregisterPlayers(FName SessionName, const TArray<FU
 			return false;
 		}
 	}
+
+#if CREATE_MIRROR_PLATFORM_SESSION
 	return BaseSessionInterface->UnregisterPlayers(SessionName, GetBaseNetIds(Players));
+#else
+	return true;
+#endif
 }
 
 FOnRegisterLocalPlayerCompleteDelegate IgnoredRegisterLocalPlayerDelegate;
@@ -869,7 +909,10 @@ void FOnlineSessionEOSPlus:: RegisterLocalPlayer(const FUniqueNetId& PlayerId, F
 	{
 		return;
 	}
+
+#if CREATE_MIRROR_PLATFORM_SESSION
 	BaseSessionInterface->RegisterLocalPlayer(*Id, SessionName, IgnoredRegisterLocalPlayerDelegate);
+#endif
 }
 
 FOnUnregisterLocalPlayerCompleteDelegate IgnoredUnregisterLocalPlayerDelegate;
@@ -893,7 +936,10 @@ void FOnlineSessionEOSPlus:: UnregisterLocalPlayer(const FUniqueNetId& PlayerId,
 	{
 		return;
 	}
+
+#if CREATE_MIRROR_PLATFORM_SESSION
 	BaseSessionInterface->UnregisterLocalPlayer(*Id, SessionName, IgnoredUnregisterLocalPlayerDelegate);
+#endif
 }
 
 int32 FOnlineSessionEOSPlus::GetNumSessions()
@@ -911,7 +957,10 @@ void FOnlineSessionEOSPlus:: DumpSessionState()
 	{
 		EOSSessionInterface->GetNumSessions();
 	}
+
+#if CREATE_MIRROR_PLATFORM_SESSION
 	BaseSessionInterface->GetNumSessions();
+#endif
 }
 
 FNamedOnlineSession* FOnlineSessionEOSPlus::GetNamedSession(FName SessionName)
