- Fix log spam from EOSPlus when base interfaces are missing

diff --git a/Engine/Plugins/Online/OnlineSubsystem/Source/Public/OnlineSubsystem.h b/Engine/Plugins/Online/OnlineSubsystem/Source/Public/OnlineSubsystem.h
--- a/Engine/Plugins/Online/OnlineSubsystem/Source/Public/OnlineSubsystem.h
+++ b/Engine/Plugins/Online/OnlineSubsystem/Source/Public/OnlineSubsystem.h
@@ -74,6 +74,13 @@ DECLARE_CYCLE_STAT_EXTERN(TEXT("VoiceInt"), STAT_Voice_Interface, STATGROUP_Onli
 	UE_CLOG(Conditional, LogOnline, Verbosity, TEXT("%s%s"), ONLINE_LOG_PREFIX, *FString::Printf(Format, ##__VA_ARGS__)); \
 }
 
+#define UE_LOG_ONLINE_ONCE(Verbosity, Format, ...) \
+{ \
+	static bool bLogged = false; \
+	UE_CLOG_ONLINE(!bLogged, Verbosity, Format, ##__VA_ARGS__); \
+	bLogged = true; \
+}
+
 /** Forward declarations of all interface classes */
 typedef TSharedPtr<class IOnlineSession, ESPMode::ThreadSafe> IOnlineSessionPtr;
 typedef TSharedPtr<class IOnlineFriends, ESPMode::ThreadSafe> IOnlineFriendsPtr;
diff --git a/Engine/Plugins/Online/OnlineSubsystemEOS/Source/OnlineSubsystemEOSPlus/Private/OnlineExternalUIEOSPlus.cpp b/Engine/Plugins/Online/OnlineSubsystemEOS/Source/OnlineSubsystemEOSPlus/Private/OnlineExternalUIEOSPlus.cpp
--- a/Engine/Plugins/Online/OnlineSubsystemEOS/Source/OnlineSubsystemEOSPlus/Private/OnlineExternalUIEOSPlus.cpp
+++ b/Engine/Plugins/Online/OnlineSubsystemEOS/Source/OnlineSubsystemEOSPlus/Private/OnlineExternalUIEOSPlus.cpp
@@ -51,7 +51,7 @@ bool FOnlineExternalUIEOSPlus::ShowLoginUI(const int ControllerIndex, bool bShow
 	}
 	else
 	{
-		UE_LOG_ONLINE(Warning, TEXT("[FOnlineExternalUIEOSPlus::ShowLoginUI] Unable to call method in base interface. Base interface not valid."));
+		UE_LOG_ONLINE_ONCE(Warning, TEXT("[FOnlineExternalUIEOSPlus::ShowLoginUI] Unable to call method in base interface. Base interface not valid."));
 	}
 
 	return bResult;
@@ -67,7 +67,7 @@ bool FOnlineExternalUIEOSPlus::ShowAccountCreationUI(const int ControllerIndex,
 	}
 	else
 	{
-		UE_LOG_ONLINE(Warning, TEXT("[FOnlineExternalUIEOSPlus::ShowAccountCreationUI] Unable to call method in base interface. Base interface not valid."));
+		UE_LOG_ONLINE_ONCE(Warning, TEXT("[FOnlineExternalUIEOSPlus::ShowAccountCreationUI] Unable to call method in base interface. Base interface not valid."));
 	}
 
 	return bResult;
@@ -83,7 +83,7 @@ bool FOnlineExternalUIEOSPlus::ShowFriendsUI(int32 LocalUserNum)
 	}
 	else
 	{
-		UE_LOG_ONLINE(Warning, TEXT("[FOnlineExternalUIEOSPlus::ShowFriendsUI] Unable to call method in base interface. Base interface not valid."));
+		UE_LOG_ONLINE_ONCE(Warning, TEXT("[FOnlineExternalUIEOSPlus::ShowFriendsUI] Unable to call method in base interface. Base interface not valid."));
 	}
 
 	return bResult;
@@ -99,7 +99,7 @@ bool FOnlineExternalUIEOSPlus::ShowInviteUI(int32 LocalUserNum, FName SessionNam
 	}
 	else
 	{
-		UE_LOG_ONLINE(Warning, TEXT("[FOnlineExternalUIEOSPlus::ShowInviteUI] Unable to call method in base interface. Base interface not valid."));
+		UE_LOG_ONLINE_ONCE(Warning, TEXT("[FOnlineExternalUIEOSPlus::ShowInviteUI] Unable to call method in base interface. Base interface not valid."));
 	}
 
 	return bResult;
@@ -115,7 +115,7 @@ bool FOnlineExternalUIEOSPlus::ShowAchievementsUI(int32 LocalUserNum)
 	}
 	else
 	{
-		UE_LOG_ONLINE(Warning, TEXT("[FOnlineExternalUIEOSPlus::ShowAchievementsUI] Unable to call method in base interface. Base interface not valid."));
+		UE_LOG_ONLINE_ONCE(Warning, TEXT("[FOnlineExternalUIEOSPlus::ShowAchievementsUI] Unable to call method in base interface. Base interface not valid."));
 	}
 
 	return bResult;
@@ -131,7 +131,7 @@ bool FOnlineExternalUIEOSPlus::ShowLeaderboardUI(const FString& LeaderboardName)
 	}
 	else
 	{
-		UE_LOG_ONLINE(Warning, TEXT("[FOnlineExternalUIEOSPlus::ShowLeaderboardUI] Unable to call method in base interface. Base interface not valid."));
+		UE_LOG_ONLINE_ONCE(Warning, TEXT("[FOnlineExternalUIEOSPlus::ShowLeaderboardUI] Unable to call method in base interface. Base interface not valid."));
 	}
 
 	return bResult;
@@ -147,7 +147,7 @@ bool FOnlineExternalUIEOSPlus::ShowWebURL(const FString& Url, const FShowWebUrlP
 	}
 	else
 	{
-		UE_LOG_ONLINE(Warning, TEXT("[FOnlineExternalUIEOSPlus::ShowWebURL] Unable to call method in base interface. Base interface not valid."));
+		UE_LOG_ONLINE_ONCE(Warning, TEXT("[FOnlineExternalUIEOSPlus::ShowWebURL] Unable to call method in base interface. Base interface not valid."));
 	}
 
 	return bResult;
@@ -163,7 +163,7 @@ bool FOnlineExternalUIEOSPlus::CloseWebURL()
 	}
 	else
 	{
-		UE_LOG_ONLINE(Warning, TEXT("[FOnlineExternalUIEOSPlus::CloseWebURL] Unable to call method in base interface. Base interface not valid."));
+		UE_LOG_ONLINE_ONCE(Warning, TEXT("[FOnlineExternalUIEOSPlus::CloseWebURL] Unable to call method in base interface. Base interface not valid."));
 	}
 
 	return bResult;
@@ -234,7 +234,7 @@ bool FOnlineExternalUIEOSPlus::ShowStoreUI(int32 LocalUserNum, const FShowStoreP
 	}
 	else
 	{
-		UE_LOG_ONLINE(Warning, TEXT("[FOnlineExternalUIEOSPlus::ShowStoreUI] Unable to call method in base interface. Base interface not valid."));
+		UE_LOG_ONLINE_ONCE(Warning, TEXT("[FOnlineExternalUIEOSPlus::ShowStoreUI] Unable to call method in base interface. Base interface not valid."));
 	}
 
 	return bResult;
@@ -250,7 +250,7 @@ bool FOnlineExternalUIEOSPlus::ShowSendMessageUI(int32 LocalUserNum, const FShow
 	}
 	else
 	{
-		UE_LOG_ONLINE(Warning, TEXT("[FOnlineExternalUIEOSPlus::ShowSendMessageUI] Unable to call method in base interface. Base interface not valid."));
+		UE_LOG_ONLINE_ONCE(Warning, TEXT("[FOnlineExternalUIEOSPlus::ShowSendMessageUI] Unable to call method in base interface. Base interface not valid."));
 	}
 
 	return bResult;
diff --git a/Engine/Plugins/Online/OnlineSubsystemEOS/Source/OnlineSubsystemEOSPlus/Private/OnlineLeaderboardsEOSPlus.cpp b/Engine/Plugins/Online/OnlineSubsystemEOS/Source/OnlineSubsystemEOSPlus/Private/OnlineLeaderboardsEOSPlus.cpp
--- a/Engine/Plugins/Online/OnlineSubsystemEOS/Source/OnlineSubsystemEOSPlus/Private/OnlineLeaderboardsEOSPlus.cpp
+++ b/Engine/Plugins/Online/OnlineSubsystemEOS/Source/OnlineSubsystemEOSPlus/Private/OnlineLeaderboardsEOSPlus.cpp
@@ -74,7 +74,7 @@ bool FOnlineLeaderboardsEOSPlus::ReadLeaderboards(const TArray< FUniqueNetIdRef
 	}
 	else
 	{
-		UE_LOG_ONLINE(Warning, TEXT("[FOnlineLeaderboardsEOSPlus::ReadLeaderboards] Unable to call method in base interface. Base interface not valid."));
+		UE_LOG_ONLINE_ONCE(Warning, TEXT("[FOnlineLeaderboardsEOSPlus::ReadLeaderboards] Unable to call method in base interface. Base interface not valid."));
 	}
 
 	EOSPlus->ExecuteNextTick([this]() {
@@ -94,7 +94,7 @@ bool FOnlineLeaderboardsEOSPlus::ReadLeaderboardsForFriends(int32 LocalUserNum,
 	}
 	else
 	{
-		UE_LOG_ONLINE(Warning, TEXT("[FOnlineLeaderboardsEOSPlus::ReadLeaderboardsForFriends] Unable to call method in base interface. Base interface not valid."));
+		UE_LOG_ONLINE_ONCE(Warning, TEXT("[FOnlineLeaderboardsEOSPlus::ReadLeaderboardsForFriends] Unable to call method in base interface. Base interface not valid."));
 
 		EOSPlus->ExecuteNextTick([this]() {
 			TriggerOnLeaderboardReadCompleteDelegates(false);
@@ -114,7 +114,7 @@ bool FOnlineLeaderboardsEOSPlus::ReadLeaderboardsAroundRank(int32 Rank, uint32 R
 	}
 	else
 	{
-		UE_LOG_ONLINE(Warning, TEXT("[FOnlineLeaderboardsEOSPlus::ReadLeaderboardsAroundRank] Unable to call method in base interface. Base interface not valid."));
+		UE_LOG_ONLINE_ONCE(Warning, TEXT("[FOnlineLeaderboardsEOSPlus::ReadLeaderboardsAroundRank] Unable to call method in base interface. Base interface not valid."));
 
 		EOSPlus->ExecuteNextTick([this]() {
 			TriggerOnLeaderboardReadCompleteDelegates(false);
@@ -160,7 +160,7 @@ void FOnlineLeaderboardsEOSPlus::FreeStats(FOnlineLeaderboardRead& ReadObject)
 	}
 	else
 	{
-		UE_LOG_ONLINE(Warning, TEXT("[FOnlineLeaderboardsEOSPlus::FreeStats] Unable to call method in base interface. Base interface not valid."));
+		UE_LOG_ONLINE_ONCE(Warning, TEXT("[FOnlineLeaderboardsEOSPlus::FreeStats] Unable to call method in base interface. Base interface not valid."));
 	}
 }
 
@@ -200,7 +200,7 @@ bool FOnlineLeaderboardsEOSPlus::FlushLeaderboards(const FName& SessionName)
 	}
 	else
 	{
-		UE_LOG_ONLINE(Warning, TEXT("[FOnlineLeaderboardsEOSPlus::FlushLeaderboards] Unable to call method in base interface. Base interface not valid."));
+		UE_LOG_ONLINE_ONCE(Warning, TEXT("[FOnlineLeaderboardsEOSPlus::FlushLeaderboards] Unable to call method in base interface. Base interface not valid."));
 
 		EOSPlus->ExecuteNextTick([this, SessionName]() {
 			TriggerOnLeaderboardFlushCompleteDelegates(SessionName, false);
@@ -220,7 +220,7 @@ bool FOnlineLeaderboardsEOSPlus::WriteOnlinePlayerRatings(const FName& SessionNa
 	}
 	else
 	{
-		UE_LOG_ONLINE(Warning, TEXT("[FOnlineLeaderboardsEOSPlus::WriteOnlinePlayerRatings] Unable to call method in base interface. Base interface not valid."));
+		UE_LOG_ONLINE_ONCE(Warning, TEXT("[FOnlineLeaderboardsEOSPlus::WriteOnlinePlayerRatings] Unable to call method in base interface. Base interface not valid."));
 	}
 
 	return bResult;
diff --git a/Engine/Plugins/Online/OnlineSubsystemEOS/Source/OnlineSubsystemEOSPlus/Private/OnlineStoreEOSPlus.cpp b/Engine/Plugins/Online/OnlineSubsystemEOS/Source/OnlineSubsystemEOSPlus/Private/OnlineStoreEOSPlus.cpp
--- a/Engine/Plugins/Online/OnlineSubsystemEOS/Source/OnlineSubsystemEOSPlus/Private/OnlineStoreEOSPlus.cpp
+++ b/Engine/Plugins/Online/OnlineSubsystemEOS/Source/OnlineSubsystemEOSPlus/Private/OnlineStoreEOSPlus.cpp
@@ -88,7 +88,7 @@ void FOnlineStoreEOSPlus::GetCategories(TArray<FOnlineStoreCategory>& OutCategor
 	}
 	else
 	{
-		UE_LOG_ONLINE(Warning, TEXT("[FOnlineStoreEOSPlus::GetCategories] Unable to call method in base interface. Base interface not valid."));
+		UE_LOG_ONLINE_ONCE(Warning, TEXT("[FOnlineStoreEOSPlus::GetCategories] Unable to call method in base interface. Base interface not valid."));
 	}
 }
 
@@ -156,7 +156,7 @@ void FOnlineStoreEOSPlus::GetOffers(TArray<FOnlineStoreOfferRef>& OutOffers) con
 	}
 	else
 	{
-		UE_LOG_ONLINE(Warning, TEXT("[FOnlineStoreEOSPlus::GetOffers] Unable to call method in base interface. Base interface not valid."));
+		UE_LOG_ONLINE_ONCE(Warning, TEXT("[FOnlineStoreEOSPlus::GetOffers] Unable to call method in base interface. Base interface not valid."));
 	}
 }
 
@@ -170,7 +170,7 @@ TSharedPtr<FOnlineStoreOffer> FOnlineStoreEOSPlus::GetOffer(const FUniqueOfferId
 	}
 	else
 	{
-		UE_LOG_ONLINE(Warning, TEXT("[FOnlineStoreEOSPlus::GetOffer] Unable to call method in base interface. Base interface not valid."));
+		UE_LOG_ONLINE_ONCE(Warning, TEXT("[FOnlineStoreEOSPlus::GetOffer] Unable to call method in base interface. Base interface not valid."));
 	}
 	
 	return Result;
diff --git a/Engine/Plugins/Online/OnlineSubsystemEOS/Source/OnlineSubsystemEOSPlus/Private/OnlineTitleFileEOSPlus.cpp b/Engine/Plugins/Online/OnlineSubsystemEOS/Source/OnlineSubsystemEOSPlus/Private/OnlineTitleFileEOSPlus.cpp
--- a/Engine/Plugins/Online/OnlineSubsystemEOS/Source/OnlineSubsystemEOSPlus/Private/OnlineTitleFileEOSPlus.cpp
+++ b/Engine/Plugins/Online/OnlineSubsystemEOS/Source/OnlineSubsystemEOSPlus/Private/OnlineTitleFileEOSPlus.cpp
@@ -53,7 +53,7 @@ bool FOnlineTitleFileEOSPlus::GetFileContents(const FString& FileName, TArray<ui
 	}
 	else
 	{
-		UE_LOG_ONLINE(VeryVerbose, TEXT("[FOnlineTitleFileEOSPlus::GetFileContents] Unable to call method in base interface. Base interface not valid"));
+		UE_LOG_ONLINE_ONCE(Warning, TEXT("[FOnlineTitleFileEOSPlus::GetFileContents] Unable to call method in base interface. Base interface not valid"));
 	}
 
 	return bResult;
@@ -69,7 +69,7 @@ bool FOnlineTitleFileEOSPlus::ClearFiles()
 	}
 	else
 	{
-		UE_LOG_ONLINE(VeryVerbose, TEXT("[FOnlineTitleFileEOSPlus::ClearFiles] Unable to call method in base interface. Base interface not valid"));
+		UE_LOG_ONLINE_ONCE(Warning, TEXT("[FOnlineTitleFileEOSPlus::ClearFiles] Unable to call method in base interface. Base interface not valid"));
 	}
 
 	return bResult;
@@ -85,7 +85,7 @@ bool FOnlineTitleFileEOSPlus::ClearFile(const FString& FileName)
 	}
 	else
 	{
-		UE_LOG_ONLINE(VeryVerbose, TEXT("[FOnlineTitleFileEOSPlus::ClearFile] Unable to call method in base interface. Base interface not valid"));
+		UE_LOG_ONLINE_ONCE(Warning, TEXT("[FOnlineTitleFileEOSPlus::ClearFile] Unable to call method in base interface. Base interface not valid"));
 	}
 
 	return bResult;
@@ -99,7 +99,7 @@ void FOnlineTitleFileEOSPlus::DeleteCachedFiles(bool bSkipEnumerated)
 	}
 	else
 	{
-		UE_LOG_ONLINE(VeryVerbose, TEXT("[FOnlineTitleFileEOSPlus::DeleteCachedFiles] Unable to call method in base interface. Base interface not valid"));
+		UE_LOG_ONLINE_ONCE(Warning, TEXT("[FOnlineTitleFileEOSPlus::DeleteCachedFiles] Unable to call method in base interface. Base interface not valid"));
 	}
 }
 
@@ -113,7 +113,7 @@ bool FOnlineTitleFileEOSPlus::EnumerateFiles(const FPagedQuery& Page)
 	}
 	else
 	{
-		UE_LOG_ONLINE(VeryVerbose, TEXT("[FOnlineTitleFileEOSPlus::EnumerateFiles] Unable to call method in base interface. Base interface not valid"));
+		UE_LOG_ONLINE_ONCE(Warning, TEXT("[FOnlineTitleFileEOSPlus::EnumerateFiles] Unable to call method in base interface. Base interface not valid"));
 	}
 
 	return bResult;
@@ -127,7 +127,7 @@ void FOnlineTitleFileEOSPlus::GetFileList(TArray<FCloudFileHeader>& Files)
 	}
 	else
 	{
-		UE_LOG_ONLINE(VeryVerbose, TEXT("[FOnlineTitleFileEOSPlus::GetFileList] Unable to call method in base interface. Base interface not valid"));
+		UE_LOG_ONLINE_ONCE(Warning, TEXT("[FOnlineTitleFileEOSPlus::GetFileList] Unable to call method in base interface. Base interface not valid"));
 	}
 }
 
@@ -141,7 +141,7 @@ bool FOnlineTitleFileEOSPlus::ReadFile(const FString& FileName)
 	}
 	else
 	{
-		UE_LOG_ONLINE(VeryVerbose, TEXT("[FOnlineTitleFileEOSPlus::ReadFile] Unable to call method in base interface. Base interface not valid"));
+		UE_LOG_ONLINE_ONCE(Warning, TEXT("[FOnlineTitleFileEOSPlus::ReadFile] Unable to call method in base interface. Base interface not valid"));
 	}
 
 	return bResult;
diff --git a/Engine/Plugins/Online/OnlineSubsystemEOS/Source/OnlineSubsystemEOSPlus/Private/OnlineUserEOSPlus.cpp b/Engine/Plugins/Online/OnlineSubsystemEOS/Source/OnlineSubsystemEOSPlus/Private/OnlineUserEOSPlus.cpp
--- a/Engine/Plugins/Online/OnlineSubsystemEOS/Source/OnlineSubsystemEOSPlus/Private/OnlineUserEOSPlus.cpp
+++ b/Engine/Plugins/Online/OnlineSubsystemEOS/Source/OnlineSubsystemEOSPlus/Private/OnlineUserEOSPlus.cpp
@@ -322,7 +322,7 @@ bool FOnlineUserEOSPlus::QueryUserInfo(int32 LocalUserNum, const TArray<FUniqueN
 	}
 	else
 	{
-		UE_LOG_ONLINE(Warning, TEXT("[FOnlineUserEOSPlus::QueryUserInfo] Unable to call method in base interface. Base interface not valid."));
+		UE_LOG_ONLINE_ONCE(Warning, TEXT("[FOnlineUserEOSPlus::QueryUserInfo] Unable to call method in base interface. Base interface not valid."));
 	}
 
 	EOSPlus->ExecuteNextTick([this, LocalUserNum, BaseUserIds]() {
@@ -349,7 +349,7 @@ bool FOnlineUserEOSPlus::GetAllUserInfo(int32 LocalUserNum, TArray<TSharedRef<FO
 	}
 	else
 	{
-		UE_LOG_ONLINE(Warning, TEXT("[FOnlineUserEOSPlus::GetAllUserInfo] Unable to call method in base interface. Base interface not valid."));
+		UE_LOG_ONLINE_ONCE(Warning, TEXT("[FOnlineUserEOSPlus::GetAllUserInfo] Unable to call method in base interface. Base interface not valid."));
 	}
 
 	return bResult;
@@ -451,7 +451,7 @@ void FOnlineUserEOSPlus::GetExternalIdMappings(const FExternalIdQueryOptions& Qu
 	}
 	else
 	{
-		UE_LOG_ONLINE(Warning, TEXT("[FOnlineUserEOSPlus::GetExternalIdMappings] Unable to call method in base interface. Base interface not valid."));
+		UE_LOG_ONLINE_ONCE(Warning, TEXT("[FOnlineUserEOSPlus::GetExternalIdMappings] Unable to call method in base interface. Base interface not valid."));
 	}
 }
 
@@ -466,7 +466,7 @@ FUniqueNetIdPtr FOnlineUserEOSPlus::GetExternalIdMapping(const FExternalIdQueryO
 	}
 	else
 	{
-		UE_LOG_ONLINE(Warning, TEXT("[FOnlineUserEOSPlus::GetExternalIdMapping] Unable to call method in base interface. Base interface not valid."));
+		UE_LOG_ONLINE_ONCE(Warning, TEXT("[FOnlineUserEOSPlus::GetExternalIdMapping] Unable to call method in base interface. Base interface not valid."));
 	}
 
 	return Result;
diff --git a/Engine/Plugins/Online/OnlineSubsystemEOS/Source/OnlineSubsystemEOSPlus/Private/OnlineVoiceEOSPlus.cpp b/Engine/Plugins/Online/OnlineSubsystemEOS/Source/OnlineSubsystemEOSPlus/Private/OnlineVoiceEOSPlus.cpp
--- a/Engine/Plugins/Online/OnlineSubsystemEOS/Source/OnlineSubsystemEOSPlus/Private/OnlineVoiceEOSPlus.cpp
+++ b/Engine/Plugins/Online/OnlineSubsystemEOS/Source/OnlineSubsystemEOSPlus/Private/OnlineVoiceEOSPlus.cpp
@@ -31,7 +31,7 @@ void FOnlineVoiceEOSPlus::Initialize()
 	}
 	else
 	{
-		UE_LOG_ONLINE(VeryVerbose, TEXT("[FOnlineVoiceEOSPlus::Initialize] BaseVoiceInterface delegates not bound. Base interface not valid"));
+		UE_LOG_ONLINE_ONCE(Warning, TEXT("[FOnlineVoiceEOSPlus::Initialize] BaseVoiceInterface delegates not bound. Base interface not valid"));
 	}
 }
 
@@ -47,7 +47,7 @@ IVoiceEnginePtr FOnlineVoiceEOSPlus::CreateVoiceEngine()
 	}
 	else
 	{
-		UE_LOG_ONLINE(VeryVerbose, TEXT("[FOnlineVoiceEOSPlus::CreateVoiceEngine] Unable to call method in base interface. Base interface not valid"));
+		UE_LOG_ONLINE_ONCE(Warning, TEXT("[FOnlineVoiceEOSPlus::CreateVoiceEngine] Unable to call method in base interface. Base interface not valid"));
 	}
 	
 	return Result;
@@ -61,7 +61,7 @@ void FOnlineVoiceEOSPlus::ProcessMuteChangeNotification()
 	}
 	else
 	{
-		UE_LOG_ONLINE(VeryVerbose, TEXT("[FOnlineVoiceEOSPlus::ProcessMuteChangeNotification] Unable to call method in base interface. Base interface not valid"));
+		UE_LOG_ONLINE_ONCE(Warning, TEXT("[FOnlineVoiceEOSPlus::ProcessMuteChangeNotification] Unable to call method in base interface. Base interface not valid"));
 	}
 }
 
@@ -73,7 +73,7 @@ void FOnlineVoiceEOSPlus::StartNetworkedVoice(uint8 LocalUserNum)
 	}
 	else
 	{
-		UE_LOG_ONLINE(VeryVerbose, TEXT("[FOnlineVoiceEOSPlus::StartNetworkedVoice] Unable to call method in base interface. Base interface not valid"));
+		UE_LOG_ONLINE_ONCE(Warning, TEXT("[FOnlineVoiceEOSPlus::StartNetworkedVoice] Unable to call method in base interface. Base interface not valid"));
 	}
 }
 
@@ -85,7 +85,7 @@ void FOnlineVoiceEOSPlus::StopNetworkedVoice(uint8 LocalUserNum)
 	}
 	else
 	{
-		UE_LOG_ONLINE(VeryVerbose, TEXT("[FOnlineVoiceEOSPlus::StopNetworkedVoice] Unable to call method in base interface. Base interface not valid"));
+		UE_LOG_ONLINE_ONCE(Warning, TEXT("[FOnlineVoiceEOSPlus::StopNetworkedVoice] Unable to call method in base interface. Base interface not valid"));
 	}
 }
 
@@ -99,7 +99,7 @@ bool FOnlineVoiceEOSPlus::RegisterLocalTalker(uint32 LocalUserNum)
 	}
 	else
 	{
-		UE_LOG_ONLINE(VeryVerbose, TEXT("[FOnlineVoiceEOSPlus::RegisterLocalTalker] Unable to call method in base interface. Base interface not valid"));
+		UE_LOG_ONLINE_ONCE(Warning, TEXT("[FOnlineVoiceEOSPlus::RegisterLocalTalker] Unable to call method in base interface. Base interface not valid"));
 	}
 
 	return bResult;
@@ -113,7 +113,7 @@ void FOnlineVoiceEOSPlus::RegisterLocalTalkers()
 	}
 	else
 	{
-		UE_LOG_ONLINE(VeryVerbose, TEXT("[FOnlineVoiceEOSPlus::RegisterLocalTalkers] Unable to call method in base interface. Base interface not valid"));
+		UE_LOG_ONLINE_ONCE(Warning, TEXT("[FOnlineVoiceEOSPlus::RegisterLocalTalkers] Unable to call method in base interface. Base interface not valid"));
 	}
 }
 
@@ -127,7 +127,7 @@ bool FOnlineVoiceEOSPlus::UnregisterLocalTalker(uint32 LocalUserNum)
 	}
 	else
 	{
-		UE_LOG_ONLINE(VeryVerbose, TEXT("[FOnlineVoiceEOSPlus::UnregisterLocalTalker] Unable to call method in base interface. Base interface not valid"));
+		UE_LOG_ONLINE_ONCE(Warning, TEXT("[FOnlineVoiceEOSPlus::UnregisterLocalTalker] Unable to call method in base interface. Base interface not valid"));
 	}
 
 	return bResult;
@@ -141,7 +141,7 @@ void FOnlineVoiceEOSPlus::UnregisterLocalTalkers()
 	}
 	else
 	{
-		UE_LOG_ONLINE(VeryVerbose, TEXT("[FOnlineVoiceEOSPlus::UnregisterLocalTalkers] Unable to call method in base interface. Base interface not valid"));
+		UE_LOG_ONLINE_ONCE(Warning, TEXT("[FOnlineVoiceEOSPlus::UnregisterLocalTalkers] Unable to call method in base interface. Base interface not valid"));
 	}
 }
 
@@ -170,7 +170,7 @@ bool FOnlineVoiceEOSPlus::RegisterRemoteTalker(const FUniqueNetId& UniqueId)
 	}
 	else
 	{
-		UE_LOG_ONLINE(VeryVerbose, TEXT("[FOnlineVoiceEOSPlus::RegisterRemoteTalker] Unable to call method in base interface. Base interface not valid"));
+		UE_LOG_ONLINE_ONCE(Warning, TEXT("[FOnlineVoiceEOSPlus::RegisterRemoteTalker] Unable to call method in base interface. Base interface not valid"));
 	}
 
 	return bResult;
@@ -201,7 +201,7 @@ bool FOnlineVoiceEOSPlus::UnregisterRemoteTalker(const FUniqueNetId& UniqueId)
 	}
 	else
 	{
-		UE_LOG_ONLINE(VeryVerbose, TEXT("[FOnlineVoiceEOSPlus::UnregisterRemoteTalker] Unable to call method in base interface. Base interface not valid"));
+		UE_LOG_ONLINE_ONCE(Warning, TEXT("[FOnlineVoiceEOSPlus::UnregisterRemoteTalker] Unable to call method in base interface. Base interface not valid"));
 	}
 
 	return bResult;
@@ -215,7 +215,7 @@ void FOnlineVoiceEOSPlus::RemoveAllRemoteTalkers()
 	}
 	else
 	{
-		UE_LOG_ONLINE(VeryVerbose, TEXT("[FOnlineVoiceEOSPlus::RemoveAllRemoteTalkers] Unable to call method in base interface. Base interface not valid"));
+		UE_LOG_ONLINE_ONCE(Warning, TEXT("[FOnlineVoiceEOSPlus::RemoveAllRemoteTalkers] Unable to call method in base interface. Base interface not valid"));
 	}
 }
 
@@ -229,7 +229,7 @@ bool FOnlineVoiceEOSPlus::IsHeadsetPresent(uint32 LocalUserNum)
 	}
 	else
 	{
-		UE_LOG_ONLINE(VeryVerbose, TEXT("[FOnlineVoiceEOSPlus::IsHeadsetPresent] Unable to call method in base interface. Base interface not valid"));
+		UE_LOG_ONLINE_ONCE(Warning, TEXT("[FOnlineVoiceEOSPlus::IsHeadsetPresent] Unable to call method in base interface. Base interface not valid"));
 	}
 
 	return bResult;
@@ -245,7 +245,7 @@ bool FOnlineVoiceEOSPlus::IsLocalPlayerTalking(uint32 LocalUserNum)
 	}
 	else
 	{
-		UE_LOG_ONLINE(VeryVerbose, TEXT("[FOnlineVoiceEOSPlus::IsLocalPlayerTalking] Unable to call method in base interface. Base interface not valid"));
+		UE_LOG_ONLINE_ONCE(Warning, TEXT("[FOnlineVoiceEOSPlus::IsLocalPlayerTalking] Unable to call method in base interface. Base interface not valid"));
 	}
 
 	return bResult;
@@ -276,7 +276,7 @@ bool FOnlineVoiceEOSPlus::IsRemotePlayerTalking(const FUniqueNetId& UniqueId)
 	}
 	else
 	{
-		UE_LOG_ONLINE(VeryVerbose, TEXT("[FOnlineVoiceEOSPlus::IsRemotePlayerTalking] Unable to call method in base interface. Base interface not valid"));
+		UE_LOG_ONLINE_ONCE(Warning, TEXT("[FOnlineVoiceEOSPlus::IsRemotePlayerTalking] Unable to call method in base interface. Base interface not valid"));
 	}
 
 	return bResult;
@@ -307,7 +307,7 @@ bool FOnlineVoiceEOSPlus::IsMuted(uint32 LocalUserNum, const FUniqueNetId& Uniqu
 	}
 	else
 	{
-		UE_LOG_ONLINE(VeryVerbose, TEXT("[FOnlineVoiceEOSPlus::IsMuted] Unable to call method in base interface. Base interface not valid"));
+		UE_LOG_ONLINE_ONCE(Warning, TEXT("[FOnlineVoiceEOSPlus::IsMuted] Unable to call method in base interface. Base interface not valid"));
 	}
 
 	return bResult;
@@ -338,7 +338,7 @@ bool FOnlineVoiceEOSPlus::MuteRemoteTalker(uint8 LocalUserNum, const FUniqueNetI
 	}
 	else
 	{
-		UE_LOG_ONLINE(VeryVerbose, TEXT("[FOnlineVoiceEOSPlus::MuteRemoteTalker] Unable to call method in base interface. Base interface not valid"));
+		UE_LOG_ONLINE_ONCE(Warning, TEXT("[FOnlineVoiceEOSPlus::MuteRemoteTalker] Unable to call method in base interface. Base interface not valid"));
 	}
 
 	return bResult;
@@ -369,7 +369,7 @@ bool FOnlineVoiceEOSPlus::UnmuteRemoteTalker(uint8 LocalUserNum, const FUniqueNe
 	}
 	else
 	{
-		UE_LOG_ONLINE(VeryVerbose, TEXT("[FOnlineVoiceEOSPlus::UnmuteRemoteTalker] Unable to call method in base interface. Base interface not valid"));
+		UE_LOG_ONLINE_ONCE(Warning, TEXT("[FOnlineVoiceEOSPlus::UnmuteRemoteTalker] Unable to call method in base interface. Base interface not valid"));
 	}
 
 	return bResult;
@@ -385,7 +385,7 @@ TSharedPtr<FVoicePacket> FOnlineVoiceEOSPlus::SerializeRemotePacket(FArchive& Ar
 	}
 	else
 	{
-		UE_LOG_ONLINE(VeryVerbose, TEXT("[FOnlineVoiceEOSPlus::SerializeRemotePacket] Unable to call method in base interface. Base interface not valid"));
+		UE_LOG_ONLINE_ONCE(Warning, TEXT("[FOnlineVoiceEOSPlus::SerializeRemotePacket] Unable to call method in base interface. Base interface not valid"));
 	}
 	
 	return Result;
@@ -401,7 +401,7 @@ TSharedPtr<FVoicePacket> FOnlineVoiceEOSPlus::GetLocalPacket(uint32 LocalUserNum
 	}
 	else
 	{
-		UE_LOG_ONLINE(VeryVerbose, TEXT("[FOnlineVoiceEOSPlus::GetLocalPacket] Unable to call method in base interface. Base interface not valid"));
+		UE_LOG_ONLINE_ONCE(Warning, TEXT("[FOnlineVoiceEOSPlus::GetLocalPacket] Unable to call method in base interface. Base interface not valid"));
 	}
 	
 	return Result;
@@ -417,7 +417,7 @@ int32 FOnlineVoiceEOSPlus::GetNumLocalTalkers()
 	}
 	else
 	{
-		UE_LOG_ONLINE(VeryVerbose, TEXT("[FOnlineVoiceEOSPlus::GetNumLocalTalkers] Unable to call method in base interface. Base interface not valid"));
+		UE_LOG_ONLINE_ONCE(Warning, TEXT("[FOnlineVoiceEOSPlus::GetNumLocalTalkers] Unable to call method in base interface. Base interface not valid"));
 	}
 
 	return Result;
@@ -431,7 +431,7 @@ void FOnlineVoiceEOSPlus::ClearVoicePackets()
 	}
 	else
 	{
-		UE_LOG_ONLINE(VeryVerbose, TEXT("[FOnlineVoiceEOSPlus::ClearVoicePackets] Unable to call method in base interface. Base interface not valid"));
+		UE_LOG_ONLINE_ONCE(Warning, TEXT("[FOnlineVoiceEOSPlus::ClearVoicePackets] Unable to call method in base interface. Base interface not valid"));
 	}
 }
 
@@ -443,7 +443,7 @@ void FOnlineVoiceEOSPlus::Tick(float DeltaTime)
 	}
 	else
 	{
-		UE_LOG_ONLINE(VeryVerbose, TEXT("[FOnlineVoiceEOSPlus::Tick] Unable to call method in base interface. Base interface not valid"));
+		UE_LOG_ONLINE_ONCE(Warning, TEXT("[FOnlineVoiceEOSPlus::Tick] Unable to call method in base interface. Base interface not valid"));
 	}
 }
 
@@ -457,7 +457,7 @@ FString FOnlineVoiceEOSPlus::GetVoiceDebugState() const
 	}
 	else
 	{
-		UE_LOG_ONLINE(VeryVerbose, TEXT("[FOnlineVoiceEOSPlus::GetVoiceDebugState] Unable to call method in base interface. Base interface not valid"));
+		UE_LOG_ONLINE_ONCE(Warning, TEXT("[FOnlineVoiceEOSPlus::GetVoiceDebugState] Unable to call method in base interface. Base interface not valid"));
 	}
 
 	return Result;
