Changes so friend invites always fire for the base platform, even if the EOS invite fails, in EOSPlus.
Added logic to cache ids for friends after ReadFriendsList in EOSPlus.

diff --git a/Engine/Plugins/Online/OnlineSubsystemEOS/Source/OnlineSubsystemEOSPlus/Private/OnlineSessionEOSPlus.cpp b/Engine/Plugins/Online/OnlineSubsystemEOS/Source/OnlineSubsystemEOSPlus/Private/OnlineSessionEOSPlus.cpp
--- a/Engine/Plugins/Online/OnlineSubsystemEOS/Source/OnlineSubsystemEOSPlus/Private/OnlineSessionEOSPlus.cpp
+++ b/Engine/Plugins/Online/OnlineSubsystemEOS/Source/OnlineSubsystemEOSPlus/Private/OnlineSessionEOSPlus.cpp
@@ -702,11 +702,11 @@ bool FOnlineSessionEOSPlus::SendSessionInviteToFriend(int32 LocalUserNum, FName
 	if (bUseEOSSessions)
 	{
 		FUniqueNetIdPtr FriendId = GetEOSNetId(Friend.ToString());
-		if (!FriendId.IsValid())
+		if (FriendId.IsValid())
 		{
-			return false;
-		}
-		return EOSSessionInterface->SendSessionInviteToFriend(LocalUserNum, SessionName, *FriendId);
+			// We want to send the platform invite always to comply with platform requirements, so we won't return here
+			EOSSessionInterface->SendSessionInviteToFriend(LocalUserNum, SessionName, *FriendId);
+		}		
 	}
 	FUniqueNetIdPtr FriendId = GetBaseNetId(Friend.ToString());
 	if (!FriendId.IsValid())
@@ -725,18 +725,21 @@ bool FOnlineSessionEOSPlus::SendSessionInviteToFriend(const FUniqueNetId& LocalU
 		{
 			return false;
 		}
+
 		FUniqueNetIdPtr FriendId = GetEOSNetId(Friend.ToString());
-		if (!FriendId.IsValid())
+		if (FriendId.IsValid())
 		{
-			return false;
-		}
-		return EOSSessionInterface->SendSessionInviteToFriend(*Id, SessionName, *FriendId);
+			// We want to send the platform invite always to comply with platform requirements, so we won't return here
+			EOSSessionInterface->SendSessionInviteToFriend(*Id, SessionName, *FriendId);
+		} 
 	}
+
 	FUniqueNetIdPtr Id = GetBaseNetId(LocalUserId.ToString());
 	if (!Id.IsValid())
 	{
 		return false;
 	}
+
 	FUniqueNetIdPtr FriendId = GetBaseNetId(Friend.ToString());
 	if (!FriendId.IsValid())
 	{
@@ -749,8 +752,10 @@ bool FOnlineSessionEOSPlus::SendSessionInviteToFriends(int32 LocalUserNum, FName
 {
 	if (bUseEOSSessions)
 	{
-		return EOSSessionInterface->SendSessionInviteToFriends(LocalUserNum, SessionName, GetEOSNetIds(Friends));
+		// We want to send the platform invite always to comply with platform requirements, so we won't return here
+		EOSSessionInterface->SendSessionInviteToFriends(LocalUserNum, SessionName, GetEOSNetIds(Friends));
 	}
+
 	return BaseSessionInterface->SendSessionInviteToFriends(LocalUserNum, SessionName, GetBaseNetIds(Friends));
 }
 
@@ -763,8 +768,11 @@ bool FOnlineSessionEOSPlus::SendSessionInviteToFriends(const FUniqueNetId& Local
 		{
 			return false;
 		}
-		return EOSSessionInterface->SendSessionInviteToFriends(*Id, SessionName, GetEOSNetIds(Friends));
+
+		// We want to send the platform invite always to comply with platform requirements, so we won't return here
+		EOSSessionInterface->SendSessionInviteToFriends(*Id, SessionName, GetEOSNetIds(Friends));
 	}
+
 	FUniqueNetIdPtr Id = GetBaseNetId(LocalUserId.ToString());
 	if (!Id.IsValid())
 	{
diff --git a/Engine/Plugins/Online/OnlineSubsystemEOS/Source/OnlineSubsystemEOSPlus/Private/OnlineUserEOSPlus.cpp b/Engine/Plugins/Online/OnlineSubsystemEOS/Source/OnlineSubsystemEOSPlus/Private/OnlineUserEOSPlus.cpp
--- a/Engine/Plugins/Online/OnlineSubsystemEOS/Source/OnlineSubsystemEOSPlus/Private/OnlineUserEOSPlus.cpp
+++ b/Engine/Plugins/Online/OnlineSubsystemEOS/Source/OnlineSubsystemEOSPlus/Private/OnlineUserEOSPlus.cpp
@@ -614,10 +614,20 @@ FUniqueNetIdEOSPlusPtr FOnlineUserEOSPlus::AddRemotePlayer(FUniqueNetIdPtr BaseN
 {
 	FUniqueNetIdEOSPlusPtr PlusNetId = FUniqueNetIdEOSPlus::Create(BaseNetId, EOSNetId);
 
-	BaseNetIdToNetIdPlus.Add(BaseNetId->ToString(), PlusNetId);
-	EOSNetIdToNetIdPlus.Add(EOSNetId->ToString(), PlusNetId);
-	NetIdPlusToBaseNetId.Add(PlusNetId->ToString(), BaseNetId);
-	NetIdPlusToEOSNetId.Add(PlusNetId->ToString(), EOSNetId);
+	// BaseNetId may not be valid for EOS friends
+	if (BaseNetId.IsValid())
+	{
+		BaseNetIdToNetIdPlus.Add(BaseNetId->ToString(), PlusNetId);
+		NetIdPlusToBaseNetId.Add(PlusNetId->ToString(), BaseNetId);
+	}
+
+	// EOSNetId may not be valid for platform friends
+	if (EOSNetId.IsValid())
+	{
+		EOSNetIdToNetIdPlus.Add(EOSNetId->ToString(), PlusNetId);
+		NetIdPlusToEOSNetId.Add(PlusNetId->ToString(), EOSNetId);
+	}
+
 	NetIdPlusToNetIdPlus.Add(PlusNetId->ToString(), PlusNetId);
 
 	return PlusNetId;
@@ -738,8 +748,21 @@ FUniqueNetIdPtr FOnlineUserEOSPlus::CreateUniquePlayerId(const FString& Str)
 		return nullptr;
 	}
 
-	FUniqueNetIdPtr BaseNetId = BaseIdentityInterface->CreateUniquePlayerId(Str.Left(FoundAt));
-	FUniqueNetIdPtr EOSNetId = EOSIdentityInterface->CreateUniquePlayerId(Str.Right(Str.Len() - FoundAt - 3));
+	// BaseNetId will only be valid for platform friends
+	FUniqueNetIdPtr BaseNetId = nullptr;
+	FString BaseNetIdStr = Str.Left(FoundAt);
+	if (!BaseNetIdStr.IsEmpty())
+	{
+		BaseNetId = BaseIdentityInterface->CreateUniquePlayerId(BaseNetIdStr);
+	}
+
+	// EOSNetId will only be valid for EOS friends
+	FUniqueNetIdPtr EOSNetId = nullptr;
+	FString EOSNetIdStr = Str.Right(Str.Len() - FoundAt - 3);
+	if (!EOSNetIdStr.IsEmpty())
+	{
+		EOSNetId = EOSIdentityInterface->CreateUniquePlayerId(EOSNetIdStr);
+	}
 
 	return AddRemotePlayer(BaseNetId, EOSNetId);
 }
@@ -933,6 +956,18 @@ void FOnlineUserEOSPlus::OnFriendRemoved(const FUniqueNetId& UserId, const FUniq
 	TriggerOnFriendRemovedDelegates(*NetIdPlus, *FriendNetIdPlus);
 }
 
+void FOnlineUserEOSPlus::CacheFriendListNetIds(int32 LocalUserNum, const FString& ListName)
+{
+	TArray<TSharedRef<FOnlineFriend>> Friends;
+	GetFriendsList(LocalUserNum, ListName, Friends);
+
+	for (const TSharedRef<FOnlineFriend>& Friend : Friends)
+	{
+		// This call will add the friends as remote players, caching their NetIds for later use
+		CreateUniquePlayerId(Friend->GetUserId()->ToString());
+	}
+}
+
 bool FOnlineUserEOSPlus::ReadFriendsList(int32 LocalUserNum, const FString& ListName, const FOnReadFriendsListComplete& Delegate)
 {
 	if (BaseFriendsInterface.IsValid())
@@ -943,6 +978,8 @@ bool FOnlineUserEOSPlus::ReadFriendsList(int32 LocalUserNum, const FString& List
 			// Skip reading EAS if not in use and if we errored at the platform level
 			if (!UEOSSettings::GetSettings().bUseEAS || !bWasSuccessful)
 			{
+				CacheFriendListNetIds(LocalUserNum, ListName);
+
 				IntermediateComplete.ExecuteIfBound(LocalUserNum, bWasSuccessful, ListName, ErrorStr);
 				return;
 			}
@@ -950,6 +987,8 @@ bool FOnlineUserEOSPlus::ReadFriendsList(int32 LocalUserNum, const FString& List
 			EOSFriendsInterface->ReadFriendsList(LocalUserNum, ListName,
 				FOnReadFriendsListComplete::CreateLambda([this, OnComplete = FOnReadFriendsListComplete(IntermediateComplete)](int32 LocalUserNum, bool bWasSuccessful, const FString& ListName, const FString& ErrorStr)
 			{
+				CacheFriendListNetIds(LocalUserNum, ListName);
+
 				OnComplete.ExecuteIfBound(LocalUserNum, bWasSuccessful, ListName, ErrorStr);
 			}));
 		}));
@@ -969,6 +1008,8 @@ bool FOnlineUserEOSPlus::ReadFriendsList(int32 LocalUserNum, const FString& List
 		return EOSFriendsInterface->ReadFriendsList(LocalUserNum, ListName,
 			FOnReadFriendsListComplete::CreateLambda([this, OnComplete = FOnReadFriendsListComplete(Delegate)](int32 LocalUserNum, bool bWasSuccessful, const FString& ListName, const FString& ErrorStr)
 		{
+			CacheFriendListNetIds(LocalUserNum, ListName);
+
 			OnComplete.ExecuteIfBound(LocalUserNum, bWasSuccessful, ListName, ErrorStr);
 		}));
 	}
diff --git a/Engine/Plugins/Online/OnlineSubsystemEOS/Source/OnlineSubsystemEOSPlus/Private/OnlineUserEOSPlus.h b/Engine/Plugins/Online/OnlineSubsystemEOS/Source/OnlineSubsystemEOSPlus/Private/OnlineUserEOSPlus.h
--- a/Engine/Plugins/Online/OnlineSubsystemEOS/Source/OnlineSubsystemEOSPlus/Private/OnlineUserEOSPlus.h
+++ b/Engine/Plugins/Online/OnlineSubsystemEOS/Source/OnlineSubsystemEOSPlus/Private/OnlineUserEOSPlus.h
@@ -482,6 +482,7 @@ private:
 	TSharedRef<FOnlineRecentPlayer> GetRecentPlayer(TSharedRef<FOnlineRecentPlayer> Player);
 	TSharedRef<FOnlineBlockedPlayer> AddBlockedPlayer(TSharedRef<FOnlineBlockedPlayer> Player);
 	TSharedRef<FOnlineBlockedPlayer> GetBlockedPlayer(TSharedRef<FOnlineBlockedPlayer> Player);
+	void CacheFriendListNetIds(int32 LocalUserNum, const FString& ListName);
 
 	/** Reference to the owning EOS plus subsystem */
 	FOnlineSubsystemEOSPlus* EOSPlus;
