Refactored LAN Sessions to use FUniqueNetIdString as SessionId like EOS Sessions and Lobbies.

diff --git a/Engine/Plugins/Online/OnlineSubsystemEOS/Source/OnlineSubsystemEOS/Private/OnlineSessionEOS.cpp b/Engine/Plugins/Online/OnlineSubsystemEOS/Source/OnlineSubsystemEOS/Private/OnlineSessionEOS.cpp
--- a/Engine/Plugins/Online/OnlineSubsystemEOS/Source/OnlineSubsystemEOS/Private/OnlineSessionEOS.cpp
+++ b/Engine/Plugins/Online/OnlineSubsystemEOS/Source/OnlineSubsystemEOS/Private/OnlineSessionEOS.cpp
@@ -406,7 +406,7 @@ void FOnlineSessionInfoEOS::InitLAN(FOnlineSubsystemEOS* Subsystem)
 
 	FGuid OwnerGuid;
 	FPlatformMisc::CreateGuid(OwnerGuid);
-	SessionId = FUniqueNetIdEOS::Create(OwnerGuid.ToString());
+	SessionId = FUniqueNetIdString::Create(OwnerGuid.ToString(), EOS_SESSION_ID);
 }
 
 typedef TEOSGlobalCallback<EOS_Sessions_OnSessionInviteReceivedCallback, EOS_Sessions_SessionInviteReceivedCallbackInfo, FOnlineSessionEOS> FSessionInviteReceivedCallback;
