Refactor of logic surrounding Epic Id queries to construct UniqueIdEOS for lobby members.
All logic making such queries now waits for them to finish before triggering any completion delegates.
Refactor of logic to mirror API Lobby data into OSS Lobbies.
This is the full-scale version of the fix submitted for OI-3978.

diff --git a/Engine/Plugins/Online/OnlineSubsystemEOS/Source/OnlineSubsystemEOS/Private/OnlineSessionEOS.cpp b/Engine/Plugins/Online/OnlineSubsystemEOS/Source/OnlineSubsystemEOS/Private/OnlineSessionEOS.cpp
--- a/Engine/Plugins/Online/OnlineSubsystemEOS/Source/OnlineSubsystemEOS/Private/OnlineSessionEOS.cpp
+++ b/Engine/Plugins/Online/OnlineSubsystemEOS/Source/OnlineSubsystemEOS/Private/OnlineSessionEOS.cpp
@@ -21,9 +21,6 @@
 	#include "eos_metrics.h"
 	#include "eos_lobby.h"
 
-static FName EOS_SESSION_ID = TEXT("EOS_SESSION_ID");
-static FName EOS_LOBBY_ID = TEXT("EOS_LOBBY_ID");
-
 /** This is the game name plus version in ansi done once for optimization */
 char BucketIdAnsi[EOS_OSS_STRING_BUFFER_LENGTH];
 
@@ -406,7 +403,7 @@ void FOnlineSessionInfoEOS::InitLAN(FOnlineSubsystemEOS* Subsystem)
 
 	FGuid OwnerGuid;
 	FPlatformMisc::CreateGuid(OwnerGuid);
-	SessionId = FUniqueNetIdString::Create(OwnerGuid.ToString(), EOS_SESSION_ID);
+	SessionId = FUniqueNetIdEOSSession::Create(OwnerGuid.ToString());
 }
 
 typedef TEOSGlobalCallback<EOS_Sessions_OnSessionInviteReceivedCallback, EOS_Sessions_SessionInviteReceivedCallbackInfo, FOnlineSessionEOS> FSessionInviteReceivedCallback;
@@ -507,7 +504,7 @@ void FOnlineSessionEOS::Init(const FString& InBucketId)
  *
  * @return pointer to the struct if found, NULL otherwise
  */
-FNamedOnlineSession* FOnlineSessionEOS::GetNamedSessionFromLobbyId(const FUniqueNetIdString& LobbyId)
+FNamedOnlineSession* FOnlineSessionEOS::GetNamedSessionFromLobbyId(const FUniqueNetIdEOSLobby& LobbyId)
 {
 	FNamedOnlineSession* Result = nullptr;
 
@@ -601,7 +598,7 @@ void FOnlineSessionEOS::RegisterLobbyNotifications()
 
 void FOnlineSessionEOS::OnLobbyUpdateReceived(const EOS_LobbyId& LobbyId)
 {
-	const FUniqueNetIdStringRef LobbyNetId = FUniqueNetIdString::Create(UTF8_TO_TCHAR(LobbyId), EOS_LOBBY_ID);
+	const FUniqueNetIdEOSLobbyRef LobbyNetId = FUniqueNetIdEOSLobby::Create(UTF8_TO_TCHAR(LobbyId));
 	FNamedOnlineSession* Session = GetNamedSessionFromLobbyId(*LobbyNetId);
 	if (Session)
 	{
@@ -615,16 +612,24 @@ void FOnlineSessionEOS::OnLobbyUpdateReceived(const EOS_LobbyId& LobbyId)
 		EOS_EResult CopyLobbyDetailsResult = EOS_Lobby_CopyLobbyDetailsHandle(LobbyHandle, &Options, &LobbyDetailsHandle);
 		if (CopyLobbyDetailsResult == EOS_EResult::EOS_Success)
 		{
+			TSharedRef<FLobbyDetailsEOS> LobbyDetails = MakeShared<FLobbyDetailsEOS>(LobbyDetailsHandle);
+
 			EOS_LobbyDetails_Info* LobbyDetailsInfo = nullptr;
 			EOS_LobbyDetails_CopyInfoOptions CopyOptions = { };
 			CopyOptions.ApiVersion = EOS_LOBBYDETAILS_COPYINFO_API_LATEST;
 
-			EOS_EResult CopyInfoResult = EOS_LobbyDetails_CopyInfo(LobbyDetailsHandle, &CopyOptions, &LobbyDetailsInfo);
+			EOS_EResult CopyInfoResult = EOS_LobbyDetails_CopyInfo(LobbyDetails->LobbyDetailsHandle, &CopyOptions, &LobbyDetailsInfo);
 			if (CopyInfoResult == EOS_EResult::EOS_Success)
 			{
-				CopyLobbyData(LobbyDetailsHandle, LobbyDetailsInfo, *Session);
-
-				TriggerOnSessionSettingsUpdatedDelegates(Session->SessionName, Session->SessionSettings);
+				CopyLobbyData(LobbyDetails, LobbyDetailsInfo, *Session, [this, SessionName = Session->SessionName, LobbyDetails](bool bWasSuccessful) {
+					if (bWasSuccessful)
+					{
+						if (FNamedOnlineSession* Session = GetNamedSession(SessionName))
+						{
+							TriggerOnSessionSettingsUpdatedDelegates(SessionName, Session->SessionSettings);
+						}
+					}
+				});
 
 				EOS_LobbyDetails_Info_Release(LobbyDetailsInfo);
 			}
@@ -632,8 +637,6 @@ void FOnlineSessionEOS::OnLobbyUpdateReceived(const EOS_LobbyId& LobbyId)
 			{
 				UE_LOG_ONLINE(Warning, TEXT("[FOnlineSessionEOS::OnLobbyUpdateReceived] EOS_LobbyDetails_CopyInfo not successful. Finished with EOS_EResult %s"), ANSI_TO_TCHAR(EOS_EResult_ToString(CopyInfoResult)));
 			}
-
-			EOS_LobbyDetails_Release(LobbyDetailsHandle);
 		}
 		else
 		{
@@ -648,55 +651,17 @@ void FOnlineSessionEOS::OnLobbyUpdateReceived(const EOS_LobbyId& LobbyId)
 
 void FOnlineSessionEOS::OnLobbyMemberUpdateReceived(const EOS_LobbyId& LobbyId, const EOS_ProductUserId& TargetUserId)
 {
-	const FUniqueNetIdStringRef LobbyNetId = FUniqueNetIdString::Create(UTF8_TO_TCHAR(LobbyId), EOS_LOBBY_ID);
+	const FUniqueNetIdEOSLobbyRef LobbyNetId = FUniqueNetIdEOSLobby::Create(UTF8_TO_TCHAR(LobbyId));
 
-	EOSSubsystem->UserManager->GetEpicAccountIdAsync(TargetUserId, [this, LobbyNetId](const EOS_ProductUserId& ProductUserId, EOS_EpicAccountId& EpicAccountId)
+	EOSSubsystem->UserManager->ResolveUniqueNetId(TargetUserId, [this, TargetUserId, LobbyNetId](FUniqueNetIdEOSRef ResolvedUniqueNetId)
 		{
-			FNamedOnlineSession* Session = GetNamedSessionFromLobbyId(*LobbyNetId);
-			if (Session)
-			{
-				FUniqueNetIdRef UniqueNetId = FUniqueNetIdEOS::Create(EpicAccountId, ProductUserId);
-
-				if (Session->SessionSettings.MemberSettings.Contains(UniqueNetId))
-				{
-					const FTCHARToUTF8 Utf8LobbyId(*LobbyNetId->ToString());
-
-					EOS_Lobby_CopyLobbyDetailsHandleOptions Options = {};
-					Options.ApiVersion = EOS_LOBBY_COPYLOBBYDETAILSHANDLE_API_LATEST;
-					Options.LobbyId = (EOS_LobbyId)Utf8LobbyId.Get();
-					Options.LocalUserId = EOSSubsystem->UserManager->GetLocalProductUserId();
-
-					EOS_HLobbyDetails LobbyDetailsHandle;
-
-					EOS_EResult Result = EOS_Lobby_CopyLobbyDetailsHandle(LobbyHandle, &Options, &LobbyDetailsHandle);
-					if (Result == EOS_EResult::EOS_Success)
-					{
-						CopyLobbyMemberAttributes(LobbyDetailsHandle, ProductUserId, Session->SessionSettings.MemberSettings[UniqueNetId]);
-
-						TriggerOnSessionParticipantSettingsUpdatedDelegates(Session->SessionName, *UniqueNetId, Session->SessionSettings);
-
-						EOS_LobbyDetails_Release(LobbyDetailsHandle);
-					}
-					else
-					{
-						UE_LOG_ONLINE(Warning, TEXT("[FOnlineSessionEOS::OnLobbyMemberUpdateReceived] EOS_LobbyDetails_CopyLobbyDetailsHandle not successful. Finished with EOS_EResult %s"), ANSI_TO_TCHAR(EOS_EResult_ToString(Result)));
-					}
-				}
-				else
-				{
-					UE_LOG_ONLINE(Warning, TEXT("[FOnlineSessionEOS::OnLobbyMemberUpdateReceived] UniqueNetId %s not registered in the session's member settings."), *UniqueNetId->ToString());
-				}
-			}
-			else
-			{
-				UE_LOG_ONLINE(Warning, TEXT("[FOnlineSessionEOS::OnLobbyMemberUpdateReceived] Unable to retrieve session with LobbyId %s"), *LobbyNetId->ToString());
-			}
+			UpdateOrAddLobbyMember(LobbyNetId, ResolvedUniqueNetId);
 		});
 }
 
 void FOnlineSessionEOS::OnMemberStatusReceived(const EOS_LobbyId& LobbyId, const EOS_ProductUserId& TargetUserId, EOS_ELobbyMemberStatus CurrentStatus)
 {
-	const FUniqueNetIdStringRef LobbyNetId = FUniqueNetIdString::Create(UTF8_TO_TCHAR(LobbyId), EOS_LOBBY_ID);
+	const FUniqueNetIdEOSLobbyRef LobbyNetId = FUniqueNetIdEOSLobby::Create(UTF8_TO_TCHAR(LobbyId));
 	FNamedOnlineSession* Session = GetNamedSessionFromLobbyId(*LobbyNetId);
 	if (Session)
 	{
@@ -704,27 +669,27 @@ void FOnlineSessionEOS::OnMemberStatusReceived(const EOS_LobbyId& LobbyId, const
 		{
 		case EOS_ELobbyMemberStatus::EOS_LMS_JOINED:
 			{
-				AddLobbyMember(LobbyNetId, TargetUserId);
+				EOSSubsystem->UserManager->ResolveUniqueNetId(TargetUserId, [this, LobbyNetId](FUniqueNetIdEOSRef ResolvedUniqueNetId)
+					{
+						UpdateOrAddLobbyMember(LobbyNetId, ResolvedUniqueNetId);
+					});
 			}
 
 			break;
 		case EOS_ELobbyMemberStatus::EOS_LMS_LEFT:
 			{
-				EOSSubsystem->UserManager->GetEpicAccountIdAsync(TargetUserId, [this, LobbyNetId](const EOS_ProductUserId& ProductUserId, EOS_EpicAccountId& EpicAccountId)
+				EOSSubsystem->UserManager->ResolveUniqueNetId(TargetUserId, [this, LobbyNetId](FUniqueNetIdEOSRef ResolvedUniqueNetId)
 					{
 						FNamedOnlineSession* Session = GetNamedSessionFromLobbyId(*LobbyNetId);
 						if (Session)
 						{
-							FUniqueNetIdRef UniqueNetId = FUniqueNetIdEOS::Create(EpicAccountId, ProductUserId);
+							RemoveOnlineSessionMember(Session->SessionName, ResolvedUniqueNetId);
 
-							UnregisterPlayer(Session->SessionName, *UniqueNetId);
-
-							// Maybe redundant with the unregister player signal?
-							TriggerOnSessionParticipantsChangeDelegates(Session->SessionName, *UniqueNetId, false);
+							TriggerOnSessionParticipantsChangeDelegates(Session->SessionName, *ResolvedUniqueNetId, false);
 						}
 						else
 						{
-							UE_LOG_ONLINE(Warning, TEXT("[FOnlineSessionEOS::OnMemberStatusReceived] Unable to retrieve session with LobbyId %s"), *LobbyNetId->ToString());
+							UE_LOG_ONLINE(Warning, TEXT("[FOnlineSessionEOS::OnMemberStatusReceived] EOS_LMS_LEFT: Unable to retrieve session with LobbyId %s"), *LobbyNetId->ToString());
 						}
 					});
 			}
@@ -734,41 +699,33 @@ void FOnlineSessionEOS::OnMemberStatusReceived(const EOS_LobbyId& LobbyId, const
 			break;
 		case EOS_ELobbyMemberStatus::EOS_LMS_KICKED:
 			{
-				EOSSubsystem->UserManager->GetEpicAccountIdAsync(TargetUserId, [this, LobbyNetId](const EOS_ProductUserId& ProductUserId, EOS_EpicAccountId& EpicAccountId)
+				EOSSubsystem->UserManager->ResolveUniqueNetId(TargetUserId, [this, LobbyNetId](FUniqueNetIdEOSRef ResolvedUniqueNetId)
 					{
 						FNamedOnlineSession* Session = GetNamedSessionFromLobbyId(*LobbyNetId);
 						if (Session)
 						{
-							int32 DefaultLocalUser = EOSSubsystem->UserManager->GetDefaultLocalUser();
-
-							FUniqueNetIdRef RemovedUserNetId = FUniqueNetIdEOS::Create(EpicAccountId, ProductUserId);
-							FUniqueNetIdPtr LocalPlayerNetId = EOSSubsystem->UserManager->GetUniquePlayerId(DefaultLocalUser);
+							RemoveOnlineSessionMember(Session->SessionName, ResolvedUniqueNetId);
 
-							if (*LocalPlayerNetId == *RemovedUserNetId)
-							{
-								TriggerOnSessionParticipantRemovedDelegates(Session->SessionName, *RemovedUserNetId);
-							}
+							TriggerOnSessionParticipantRemovedDelegates(Session->SessionName, *ResolvedUniqueNetId);
 						}
 						else
 						{
-							UE_LOG_ONLINE(Warning, TEXT("[FOnlineSessionEOS::OnMemberStatusReceived] Unable to retrieve session with LobbyId %s"), *LobbyNetId->ToString());
+							UE_LOG_ONLINE(Warning, TEXT("[FOnlineSessionEOS::OnMemberStatusReceived] EOS_LMS_KICKED: Unable to retrieve session with LobbyId %s"), *LobbyNetId->ToString());
 						}
 					});
 			}
 			break;
 		case EOS_ELobbyMemberStatus::EOS_LMS_PROMOTED:
 			{
-				EOSSubsystem->UserManager->GetEpicAccountIdAsync(TargetUserId, [this, LobbyNetId](const EOS_ProductUserId& ProductUserId, EOS_EpicAccountId& EpicAccountId)
+				EOSSubsystem->UserManager->ResolveUniqueNetId(TargetUserId, [this, LobbyNetId](FUniqueNetIdEOSRef ResolvedUniqueNetId)
 					{
 						FNamedOnlineSession* Session = GetNamedSessionFromLobbyId(*LobbyNetId);
 						if (Session)
 						{
 							int32 DefaultLocalUser = EOSSubsystem->UserManager->GetDefaultLocalUser();
-
-							FUniqueNetIdRef NewOwnerUniqueNetId = FUniqueNetIdEOS::Create(EpicAccountId, ProductUserId);
 							FUniqueNetIdPtr LocalPlayerUniqueNetId = EOSSubsystem->UserManager->GetUniquePlayerId(DefaultLocalUser);
 
-							if (*LocalPlayerUniqueNetId == *NewOwnerUniqueNetId)
+							if (*LocalPlayerUniqueNetId == *ResolvedUniqueNetId)
 							{
 								Session->OwningUserId = LocalPlayerUniqueNetId;
 								Session->OwningUserName = EOSSubsystem->UserManager->GetPlayerNickname(DefaultLocalUser);;
@@ -781,7 +738,7 @@ void FOnlineSessionEOS::OnMemberStatusReceived(const EOS_LobbyId& LobbyId, const
 						}
 						else
 						{
-							UE_LOG_ONLINE(Warning, TEXT("[FOnlineSessionEOS::OnMemberStatusReceived] Unable to retrieve session with LobbyId %s"), *LobbyNetId->ToString());
+							UE_LOG_ONLINE(Warning, TEXT("[FOnlineSessionEOS::OnMemberStatusReceived] EOS_LMS_PROMOTED: Unable to retrieve session with LobbyId %s"), *LobbyNetId->ToString());
 						}
 					});
 			}
@@ -817,9 +774,14 @@ void FOnlineSessionEOS::OnLobbyInviteAccepted(const char* InviteId, const EOS_Pr
 	EOS_EResult Result = EOS_Lobby_CopyLobbyDetailsHandleByInviteId(LobbyHandle, &Options, &LobbyDetailsHandle);
 	if (Result == EOS_EResult::EOS_Success)
 	{
+		TSharedRef<FLobbyDetailsEOS> LobbyDetails = MakeShared<FLobbyDetailsEOS>(LobbyDetailsHandle);
+
 		LastInviteSearch = MakeShared<FOnlineSessionSearch>();
-		AddLobbySearchResult(LobbyDetailsHandle, LastInviteSearch.ToSharedRef());
-		TriggerOnSessionUserInviteAcceptedDelegates(true, LocalUserNum, NetId, LastInviteSearch->SearchResults.Last());
+		AddLobbySearchResult(LobbyDetails, LastInviteSearch.ToSharedRef(), [this, LocalUserNum, NetId](bool bWasSuccessful)
+		{
+			// If we fail to copy the lobby data, we won't add a new search result, so we'll return an empty one
+			TriggerOnSessionUserInviteAcceptedDelegates(bWasSuccessful, LocalUserNum, NetId, bWasSuccessful ? LastInviteSearch->SearchResults.Last() : FOnlineSessionSearchResult());
+		});
 	}
 	else
 	{
@@ -847,9 +809,14 @@ void FOnlineSessionEOS::OnJoinLobbyAccepted(const EOS_ProductUserId& LocalUserId
 	EOS_EResult Result = EOS_Lobby_CopyLobbyDetailsHandleByUiEventId(LobbyHandle, &Options, &LobbyDetailsHandle);
 	if (Result == EOS_EResult::EOS_Success)
 	{
+		TSharedRef<FLobbyDetailsEOS> LobbyDetails = MakeShared<FLobbyDetailsEOS>(LobbyDetailsHandle);
+
 		LastInviteSearch = MakeShared<FOnlineSessionSearch>();
-		AddLobbySearchResult(LobbyDetailsHandle, LastInviteSearch.ToSharedRef());
-		TriggerOnSessionUserInviteAcceptedDelegates(true, LocalUserNum, NetId, LastInviteSearch->SearchResults.Last());
+		AddLobbySearchResult(LobbyDetails, LastInviteSearch.ToSharedRef(), [this, LocalUserNum, NetId](bool bWasSuccessful)
+		{
+			// If we fail to copy the lobby data, we won't add a new search result, so we'll return an empty one
+			TriggerOnSessionUserInviteAcceptedDelegates(bWasSuccessful, LocalUserNum, NetId, bWasSuccessful ? LastInviteSearch->SearchResults.Last() : FOnlineSessionSearchResult());
+		});
 	}
 	else
 	{
@@ -1219,7 +1186,7 @@ uint32 FOnlineSessionEOS::CreateEOSSession(int32 HostingPlayerNum, FNamedOnlineS
 		// This is basically ignored
 		HostAddr = TEXT("127.0.0.1");
 	}
-	Session->SessionInfo = MakeShareable(new FOnlineSessionInfoEOS(HostAddr, FUniqueNetIdString::Create(FString(), EOS_SESSION_ID), nullptr));
+	Session->SessionInfo = MakeShareable(new FOnlineSessionInfoEOS(HostAddr, FUniqueNetIdEOSSession::Create(FString()), nullptr));
 
 	FName SessionName = Session->SessionName;
 
@@ -1237,7 +1204,7 @@ uint32 FOnlineSessionEOS::CreateEOSSession(int32 HostingPlayerNum, FNamedOnlineS
 				TSharedPtr<FOnlineSessionInfoEOS> SessionInfo = StaticCastSharedPtr<FOnlineSessionInfoEOS>(Session->SessionInfo);
 				if (SessionInfo.IsValid())
 				{
-					SessionInfo->SessionId = FUniqueNetIdString::Create(UTF8_TO_TCHAR(Data->SessionId), EOS_SESSION_ID);
+					SessionInfo->SessionId = FUniqueNetIdEOSSession::Create(UTF8_TO_TCHAR(Data->SessionId));
 				}
 
 				Session->SessionState = EOnlineSessionState::Pending;
@@ -1838,7 +1805,7 @@ bool FOnlineSessionEOS::FindSessionById(const FUniqueNetId& SearchingUserId, con
 		CurrentSessionSearch->SearchState = EOnlineAsyncTaskState::InProgress;
 
 		StartLobbySearch(EOSSubsystem->UserManager->GetLocalUserNumFromUniqueNetId(SearchingUserId), LobbySearchHandle, CurrentSessionSearch.ToSharedRef(),
-			FOnSingleSessionResultCompleteDelegate::CreateLambda([this, OrigCallback = FOnSingleSessionResultCompleteDelegate(CompletionDelegate), SessId = FUniqueNetIdString::Create(SessionId.ToString(), EOS_SESSION_ID)](int32 LocalUserNum, bool bWasSuccessful, const FOnlineSessionSearchResult& EOSResult)
+			FOnSingleSessionResultCompleteDelegate::CreateLambda([this, OrigCallback = FOnSingleSessionResultCompleteDelegate(CompletionDelegate), SessId = FUniqueNetIdEOSSession::Create(SessionId.ToString())](int32 LocalUserNum, bool bWasSuccessful, const FOnlineSessionSearchResult& EOSResult)
 		{
 			if (bWasSuccessful)
 			{
@@ -2017,7 +1984,7 @@ void FOnlineSessionEOS::AddSearchResult(EOS_HSessionDetails SessionHandle, const
 		int32 Position = SearchSettings->SearchResults.AddZeroed();
 		FOnlineSessionSearchResult& SearchResult = SearchSettings->SearchResults[Position];
 		// This will set the host address and port
-		SearchResult.Session.SessionInfo = MakeShareable(new FOnlineSessionInfoEOS(SessionInfo->HostAddress, FUniqueNetIdString::Create(SessionInfo->SessionId, EOS_SESSION_ID), SessionHandle));
+		SearchResult.Session.SessionInfo = MakeShareable(new FOnlineSessionInfoEOS(SessionInfo->HostAddress, FUniqueNetIdEOSSession::Create(SessionInfo->SessionId), SessionHandle));
 
 		CopySearchResult(SessionHandle, SessionInfo, SearchResult.Session);
 
@@ -2747,6 +2714,124 @@ void FOnlineSessionEOS::RegisterLocalPlayers(FNamedOnlineSession* Session)
 
 }
 
+void FOnlineSessionEOS::UpdateOrAddLobbyMember(const FUniqueNetIdEOSLobbyRef& LobbyNetId, const FUniqueNetIdEOSRef& PlayerId)
+{
+	if (FNamedOnlineSession* Session = GetNamedSessionFromLobbyId(*LobbyNetId))
+	{
+		// First we add the player to the session, if it wasn't already there
+		bool bWasLobbyMemberAdded = false;
+		if (!Session->SessionSettings.MemberSettings.Contains(PlayerId))
+		{
+			bWasLobbyMemberAdded = AddOnlineSessionMember(Session->SessionName, PlayerId);
+		}
+
+		if (FSessionSettings* MemberSettings = Session->SessionSettings.MemberSettings.Find(PlayerId))
+		{
+			const FTCHARToUTF8 Utf8LobbyId(*LobbyNetId->ToString());
+
+			EOS_Lobby_CopyLobbyDetailsHandleOptions Options = {};
+			Options.ApiVersion = EOS_LOBBY_COPYLOBBYDETAILSHANDLE_API_LATEST;
+			Options.LobbyId = (EOS_LobbyId)Utf8LobbyId.Get();
+			Options.LocalUserId = EOSSubsystem->UserManager->GetLocalProductUserId();
+
+			EOS_HLobbyDetails LobbyDetailsHandle;
+
+			EOS_EResult Result = EOS_Lobby_CopyLobbyDetailsHandle(LobbyHandle, &Options, &LobbyDetailsHandle);
+			if (Result == EOS_EResult::EOS_Success)
+			{
+				FLobbyDetailsEOS LobbyDetails(LobbyDetailsHandle);
+
+				// Then we update their attributes
+				CopyLobbyMemberAttributes(LobbyDetails, PlayerId->GetProductUserId(), *MemberSettings);
+
+				if (bWasLobbyMemberAdded)
+				{
+					TriggerOnSessionParticipantsChangeDelegates(Session->SessionName, *PlayerId, true);
+				}
+				else
+				{
+					TriggerOnSessionParticipantSettingsUpdatedDelegates(Session->SessionName, *PlayerId, Session->SessionSettings);
+				}
+			}
+			else
+			{
+				UE_LOG_ONLINE(Warning, TEXT("[FOnlineSessionEOS::UpdateOrAddLobbyMember] EOS_LobbyDetails_CopyLobbyDetailsHandle not successful. Finished with EOS_EResult %s"), ANSI_TO_TCHAR(EOS_EResult_ToString(Result)));
+			}
+		}
+		else
+		{
+			UE_LOG_ONLINE(Warning, TEXT("[FOnlineSessionEOS::UpdateOrAddLobbyMember] UniqueNetId %s not registered in the session's member settings."), *PlayerId->ToString());
+		}
+	}
+	else
+	{
+		UE_LOG_ONLINE(Warning, TEXT("[FOnlineSessionEOS::UpdateOrAddLobbyMember] Unable to retrieve session with LobbyId %s"), *LobbyNetId->ToString());
+	}
+}
+
+bool FOnlineSessionEOS::AddOnlineSessionMember(FName SessionName, const FUniqueNetIdRef& PlayerId)
+{
+	FNamedOnlineSession* Session = GetNamedSession(SessionName);
+	if (Session)
+	{
+		if (!Session->SessionSettings.MemberSettings.Contains(PlayerId))
+		{
+			// update number of open connections
+			if (Session->NumOpenPublicConnections > 0)
+			{
+				Session->NumOpenPublicConnections--;
+			}
+			else if (Session->NumOpenPrivateConnections > 0)
+			{
+				Session->NumOpenPrivateConnections--;
+			}
+			else
+			{
+				UE_LOG_ONLINE_SESSION(Warning, TEXT("[FOnlineSessionEOS::AddOnlineSessionMember] Could not add new member to session %s, no Public or Private connections open"), *SessionName.ToString());
+
+				return false;
+			}
+
+			Session->SessionSettings.MemberSettings.Add(PlayerId, FSessionSettings());
+
+			return true;
+		}
+	}
+	else
+	{
+		UE_LOG_ONLINE_SESSION(Warning, TEXT("[FOnlineSessionEOS::AddOnlineSessionMember] Could not find session with name: %s"), *SessionName.ToString());
+	}
+
+	return false;
+}
+
+bool FOnlineSessionEOS::RemoveOnlineSessionMember(FName SessionName, const FUniqueNetIdRef& PlayerId)
+{
+	FNamedOnlineSession* Session = GetNamedSession(SessionName);
+	if (Session)
+	{
+		// update number of open connections
+		if (Session->NumOpenPublicConnections < Session->SessionSettings.NumPublicConnections)
+		{
+			Session->NumOpenPublicConnections++;
+		}
+		else if (Session->NumOpenPrivateConnections < Session->SessionSettings.NumPrivateConnections)
+		{
+			Session->NumOpenPrivateConnections++;
+		}
+
+		Session->SessionSettings.MemberSettings.Remove(PlayerId);
+
+		return true;
+	}
+	else
+	{
+		UE_LOG_ONLINE_SESSION(Warning, TEXT("[FOnlineSessionEOS::RemoveOnlineSessionMember] Could not find session with name: %s"), *SessionName.ToString());
+	}
+
+	return false;
+}
+
 bool FOnlineSessionEOS::RegisterPlayer(FName SessionName, const FUniqueNetId& PlayerId, bool bWasInvited)
 {
 	TArray< FUniqueNetIdRef > Players;
@@ -2780,26 +2865,14 @@ bool FOnlineSessionEOS::RegisterPlayers(FName SessionName, const TArray< FUnique
 					EOSIds.Add(PlayerEOSId.GetProductUserId());
 				}
 
-				// update number of open connections
-				if (Session->NumOpenPublicConnections > 0)
-				{
-					Session->NumOpenPublicConnections--;
-				}
-				else if (Session->NumOpenPrivateConnections > 0)
-				{
-					Session->NumOpenPrivateConnections--;
-				}
-
-				if (!Session->SessionSettings.MemberSettings.Contains(PlayerId))
-				{
-					Session->SessionSettings.MemberSettings.Add(PlayerId, FSessionSettings());
-				}
+				AddOnlineSessionMember(SessionName, PlayerId);
 			}
 			else
 			{
 				UE_LOG_ONLINE_SESSION(Log, TEXT("Player %s already registered in session %s"), *PlayerId->ToDebugString(), *SessionName.ToString());
 			}
 		}
+
 		if (bRegisterEOS && EOSIds.Num() > 0)
 		{
 			EOS_Sessions_RegisterPlayersOptions Options = { };
@@ -2865,20 +2938,7 @@ bool FOnlineSessionEOS::UnregisterPlayers(FName SessionName, const TArray< FUniq
 					EOSIds.Add(PlayerEOSId.GetProductUserId());
 				}
 
-				// update number of open connections
-				if (Session->NumOpenPublicConnections < Session->SessionSettings.NumPublicConnections)
-				{
-					Session->NumOpenPublicConnections++;
-				}
-				else if (Session->NumOpenPrivateConnections < Session->SessionSettings.NumPrivateConnections)
-				{
-					Session->NumOpenPrivateConnections++;
-				}
-
-				if (Session->SessionSettings.MemberSettings.Contains(PlayerId))
-				{
-					Session->SessionSettings.MemberSettings.Remove(PlayerId);
-				}
+				RemoveOnlineSessionMember(SessionName, PlayerId);
 			}
 			else
 			{
@@ -3293,7 +3353,7 @@ bool FOnlineSessionEOS::IsHost(const FNamedOnlineSession& Session) const
 
 FUniqueNetIdPtr FOnlineSessionEOS::CreateSessionIdFromString(const FString& SessionIdStr)
 {
-	return FUniqueNetIdString::Create(SessionIdStr, EOS_SESSION_ID);
+	return FUniqueNetIdEOSSession::Create(SessionIdStr);
 }
 
 EOS_ELobbyPermissionLevel FOnlineSessionEOS::GetLobbyPermissionLevelFromSessionSettings(const FOnlineSessionSettings& SessionSettings)
@@ -3378,7 +3438,7 @@ uint32 FOnlineSessionEOS::CreateLobbySession(int32 HostingPlayerNum, FNamedOnlin
 				FInternetAddrEOS TempAddr(LexToString(LocalProductUserId), SessionName.ToString(), FURL::UrlConfig.DefaultPort);
 				FString HostAddr = TempAddr.ToString(true);
 
-				Session->SessionInfo = MakeShareable(new FOnlineSessionInfoEOS(HostAddr, FUniqueNetIdString::Create(Data->LobbyId, EOS_LOBBY_ID), nullptr));
+				Session->SessionInfo = MakeShareable(new FOnlineSessionInfoEOS(HostAddr, FUniqueNetIdEOSLobby::Create(Data->LobbyId), nullptr));
 
 #if WITH_EOS_RTC
 				if (FEOSVoiceChatUser* VoiceChatUser = static_cast<FEOSVoiceChatUser*>(EOSSubsystem->GetEOSVoiceChatUserInterface(*LocalUserNetId)))
@@ -3429,20 +3489,21 @@ uint32 FOnlineSessionEOS::JoinLobbySession(int32 PlayerNum, FNamedOnlineSession*
 
 			Session->SessionState = EOnlineSessionState::Pending;
 			
+			// We retrieve the cached LobbyDetailsHandle and we start the join operation
 			EOS_Lobby_JoinLobbyOptions JoinLobbyOptions = { 0 };
 			JoinLobbyOptions.ApiVersion = EOS_LOBBY_JOINLOBBY_API_LATEST;
 			JoinLobbyOptions.LocalUserId = EOSSubsystem->UserManager->GetLocalProductUserId(PlayerNum);
 			JoinLobbyOptions.bPresenceEnabled = Session->SessionSettings.bUsesPresence;
 
- 			EOS_HLobbyDetails LobbyDetailsHandle = *LobbySearchResultsCache[Session->SessionInfo->GetSessionId().ToString()];
- 			JoinLobbyOptions.LobbyDetailsHandle = LobbyDetailsHandle;
+ 			TSharedRef<FLobbyDetailsEOS> LobbyDetails = LobbySearchResultsCache[Session->SessionInfo->GetSessionId().ToString()];
+ 			JoinLobbyOptions.LobbyDetailsHandle = LobbyDetails->LobbyDetailsHandle;
 
 			FName SessionName = Session->SessionName;
 			FUniqueNetIdPtr LocalUserNetId = EOSSubsystem->UserManager->GetLocalUniqueNetIdEOS(PlayerNum);
 
 			FLobbyJoinedCallback* CallbackObj = new FLobbyJoinedCallback(FOnlineSessionEOSWeakPtr(AsShared()));
 			LobbyJoinedCallback = CallbackObj;
-			CallbackObj->CallbackLambda = [this, SessionName, LocalUserNetId, LobbyDetailsHandle](const EOS_Lobby_JoinLobbyCallbackInfo* Data)
+			CallbackObj->CallbackLambda = [this, SessionName, LocalUserNetId](const EOS_Lobby_JoinLobbyCallbackInfo* Data)
 			{
 				FNamedOnlineSession* Session = GetNamedSession(SessionName);
 				if (Session)
@@ -3453,21 +3514,6 @@ uint32 FOnlineSessionEOS::JoinLobbySession(int32 PlayerNum, FNamedOnlineSession*
 						UE_LOG_ONLINE_SESSION(Verbose, TEXT("[FOnlineSessionEOS::JoinLobbySession] JoinLobby was successful. LobbyId is %d."), Data->LobbyId);
 
 						BeginSessionAnalytics(Session);
-						
-						// Initialize the OSS's member list with the current member list of the lobby.
-						const FUniqueNetIdStringRef LobbyNetId = FUniqueNetIdString::Create(UTF8_TO_TCHAR(Data->LobbyId), EOS_LOBBY_ID);
-						EOS_LobbyDetails_GetMemberCountOptions GetMemberCountOptions = { 0 };
-						GetMemberCountOptions.ApiVersion = EOS_LOBBYDETAILS_GETMEMBERCOUNT_API_LATEST;
-						const uint32 MemberCount = EOS_LobbyDetails_GetMemberCount(LobbyDetailsHandle, &GetMemberCountOptions);
-						
-						EOS_LobbyDetails_GetMemberByIndexOptions GetMemberOptions = { 0 };
-						GetMemberOptions.ApiVersion = EOS_LOBBYDETAILS_GETMEMBERBYINDEX_API_LATEST;
-						for(uint32 MemberIndex = 0; MemberIndex < MemberCount; ++MemberIndex)
-						{
-							GetMemberOptions.MemberIndex = MemberIndex;
-							const EOS_ProductUserId ProductUserId = EOS_LobbyDetails_GetMemberByIndex(LobbyDetailsHandle, &GetMemberOptions);
-							AddLobbyMember(LobbyNetId, ProductUserId);
-						}
 
 #if WITH_EOS_RTC
 						if (FEOSVoiceChatUser* VoiceChatUser = static_cast<FEOSVoiceChatUser*>(EOSSubsystem->GetEOSVoiceChatUserInterface(*LocalUserNetId)))
@@ -3486,26 +3532,25 @@ uint32 FOnlineSessionEOS::JoinLobbySession(int32 PlayerNum, FNamedOnlineSession*
 					}
 
 					TriggerOnJoinSessionCompleteDelegates(SessionName, bWasSuccessful ? EOnJoinSessionCompleteResult::Success : EOnJoinSessionCompleteResult::UnknownError);
-
-					// Now that we have joined one of the sessions we found on the search, we can clear the results
-					for (TPair<FString, TSharedRef<EOS_HLobbyDetails>> LobbySearchResult : LobbySearchResultsCache)
-					{
-						EOS_LobbyDetails_Release(*LobbySearchResult.Value);
-					}
-
-					LobbySearchResultsCache.Reset();
 				}
 				else
 				{
 					UE_LOG_ONLINE_SESSION(Warning, TEXT("[FOnlineSessionEOS::JoinLobbySession] Unable to find session %s"), *SessionName.ToString());
 					TriggerOnJoinSessionCompleteDelegates(SessionName, EOnJoinSessionCompleteResult::SessionDoesNotExist);
 				}
+
+				// Now that we have finished the join process, we can clear the cache
+				LobbySearchResultsCache.Reset();
 			};
 
 			EOS_Lobby_JoinLobby(LobbyHandle, &JoinLobbyOptions, CallbackObj, CallbackObj->GetCallbackPtr());
 
 			Result = ONLINE_IO_PENDING;
 		}
+		else
+		{
+			UE_LOG_ONLINE(Warning, TEXT("[FOnlineSessionEOS::JoinLobbySession] SessionId not valid"));
+		}
 	}
 	else
 	{
@@ -3654,31 +3699,6 @@ void FOnlineSessionEOS::SetLobbyAttributes(EOS_HLobbyModification LobbyModificat
 	}
 }
 
-void FOnlineSessionEOS::AddLobbyMember(const FUniqueNetIdStringRef LobbyNetId, const EOS_ProductUserId& TargetUserId)
-{
-	EOSSubsystem->UserManager->GetEpicAccountIdAsync(TargetUserId, [this, LobbyNetId](const EOS_ProductUserId& ProductUserId, EOS_EpicAccountId& EpicAccountId)
-		{
-			if (const FNamedOnlineSession* Session = GetNamedSessionFromLobbyId(*LobbyNetId))
-			{
-				const FUniqueNetIdRef UniqueNetId = FUniqueNetIdEOS::Create(EpicAccountId, ProductUserId);
-
-				RegisterPlayer(Session->SessionName, *UniqueNetId, false);
-
-				// Right after registering the user, we query their member settings
-				const FTCHARToUTF8 Utf8LobbyId(*LobbyNetId->ToString());
-				OnLobbyMemberUpdateReceived(static_cast<EOS_LobbyId>(Utf8LobbyId.Get()), ProductUserId);
-
-				// Maybe redundant with the register player signal?
-				TriggerOnSessionParticipantsChangeDelegates(Session->SessionName, *UniqueNetId, true);
-			}
-			else
-			{
-				UE_LOG_ONLINE(Warning, TEXT("[FOnlineSessionEOS::AddLobbyMember] Unable to retrieve session with LobbyId %s"), *LobbyNetId->ToString());
-			}
-		});
-}
-
-
 uint32 FOnlineSessionEOS::UpdateLobbySession(FNamedOnlineSession* Session)
 {
 	check(Session != nullptr);
@@ -3880,6 +3900,9 @@ uint32 FOnlineSessionEOS::FindLobbySession(int32 SearchingPlayerNum, const TShar
 
 void FOnlineSessionEOS::StartLobbySearch(int32 SearchingPlayerNum, EOS_HLobbySearch LobbySearchHandle, const TSharedRef<FOnlineSessionSearch>& SearchSettings, const FOnSingleSessionResultCompleteDelegate& CompletionDelegate)
 {
+	// When starting a new search, we'll reset the cache
+	LobbySearchResultsCache.Reset();
+
 	SessionSearchStartInSeconds = FPlatformTime::Seconds();
 
 	EOS_LobbySearch_FindOptions FindOptions = { 0 };
@@ -3894,8 +3917,6 @@ void FOnlineSessionEOS::StartLobbySearch(int32 SearchingPlayerNum, EOS_HLobbySea
 		{
 			UE_LOG_ONLINE_SESSION(Log, TEXT("[FOnlineSessionEOS::StartLobbySearch] LobbySearch_Find was successful."));
 
-			LobbySearchResultsCache.Reset();
-
 			CurrentSessionSearch->SearchState = EOnlineAsyncTaskState::Done;
 
 			EOS_LobbySearch_GetSearchResultCountOptions GetSearchResultCountOptions = { 0 };
@@ -3917,9 +3938,22 @@ void FOnlineSessionEOS::StartLobbySearch(int32 SearchingPlayerNum, EOS_HLobbySea
 					EOS_EResult Result = EOS_LobbySearch_CopySearchResultByIndex(LobbySearchHandle, &CopySearchResultByIndexOptions, &LobbyDetailsHandle);
 					if (Result == EOS_EResult::EOS_Success)
 					{
+						TSharedRef<FLobbyDetailsEOS> LobbyDetails = MakeShared<FLobbyDetailsEOS>(LobbyDetailsHandle);
+
 						UE_LOG_ONLINE_SESSION(Verbose, TEXT("[FOnlineSessionEOS::StartLobbySearch::FLobbySearchFindCallback] LobbySearch_CopySearchResultByIndex was successful."));
 
-						AddLobbySearchResult(LobbyDetailsHandle, SearchSettings);
+						PendingLobbySearchResults.Add(LobbyDetails);
+
+						AddLobbySearchResult(LobbyDetails, SearchSettings, [this, LobbyDetails, CompletionDelegate, SearchingPlayerNum, SearchSettings](bool bWasSuccessful)
+						{
+							PendingLobbySearchResults.Remove(LobbyDetails);
+
+							if (PendingLobbySearchResults.IsEmpty())
+							{
+								// If we fail to copy the lobby data, we won't add a new search result, so we'll return an empty one
+								CompletionDelegate.ExecuteIfBound(SearchingPlayerNum, bWasSuccessful, bWasSuccessful ? SearchSettings->SearchResults.Last() : FOnlineSessionSearchResult());
+							}
+						});
 					}
 					else
 					{
@@ -3927,7 +3961,10 @@ void FOnlineSessionEOS::StartLobbySearch(int32 SearchingPlayerNum, EOS_HLobbySea
 					}
 				}
 
-				CompletionDelegate.ExecuteIfBound(SearchingPlayerNum, true, SearchSettings->SearchResults.Last());
+				if (PendingLobbySearchResults.IsEmpty())
+				{
+					CompletionDelegate.ExecuteIfBound(SearchingPlayerNum, true, SearchSettings->SearchResults.Last());
+				}
 			}
 			else
 			{
@@ -3951,12 +3988,12 @@ void FOnlineSessionEOS::StartLobbySearch(int32 SearchingPlayerNum, EOS_HLobbySea
 	EOS_LobbySearch_Find(LobbySearchHandle, &FindOptions, CallbackObj, CallbackObj->GetCallbackPtr());
 }
 
-void FOnlineSessionEOS::AddLobbySearchResult(EOS_HLobbyDetails LobbyDetailsHandle, const TSharedRef<FOnlineSessionSearch>& SearchSettings)
+void FOnlineSessionEOS::AddLobbySearchResult(const TSharedRef<FLobbyDetailsEOS>& LobbyDetails, const TSharedRef<FOnlineSessionSearch>& SearchSettings, const FOnCopyLobbyDataCompleteCallback& Callback)
 {
 	EOS_LobbyDetails_Info* LobbyDetailsInfo = nullptr;
 	EOS_LobbyDetails_CopyInfoOptions CopyOptions = { };
 	CopyOptions.ApiVersion = EOS_LOBBYDETAILS_COPYINFO_API_LATEST;
-	EOS_EResult CopyResult = EOS_LobbyDetails_CopyInfo(LobbyDetailsHandle, &CopyOptions, &LobbyDetailsInfo);
+	EOS_EResult CopyResult = EOS_LobbyDetails_CopyInfo(LobbyDetails->LobbyDetailsHandle, &CopyOptions, &LobbyDetailsInfo);
 	if (CopyResult == EOS_EResult::EOS_Success)
 	{
 		int32 Position = SearchSettings->SearchResults.AddZeroed();
@@ -3969,39 +4006,12 @@ void FOnlineSessionEOS::AddLobbySearchResult(EOS_HLobbyDetails LobbyDetailsHandl
 		FInternetAddrEOS TempAddr(LexToString(LobbyDetailsInfo->LobbyOwnerUserId), NetDriverName.ToString(), GetTypeHash(NetDriverName.ToString()));
 		FString HostAddr = TempAddr.ToString(true);
 
-		SearchResult.Session.SessionInfo = MakeShareable(new FOnlineSessionInfoEOS(HostAddr, FUniqueNetIdString::Create(LobbyDetailsInfo->LobbyId, EOS_LOBBY_ID), nullptr));
+		SearchResult.Session.SessionInfo = MakeShareable(new FOnlineSessionInfoEOS(HostAddr, FUniqueNetIdEOSLobby::Create(LobbyDetailsInfo->LobbyId), nullptr));
 
 		// We copy the lobby data and settings
-		CopyLobbyData(LobbyDetailsHandle, LobbyDetailsInfo, SearchResult.Session);
-
-		// Then we copy the settings for all lobby members
-		EOS_LobbyDetails_GetMemberCountOptions CountOptions = { };
-		CountOptions.ApiVersion = EOS_LOBBYDETAILS_GETMEMBERCOUNT_API_LATEST;
-		int32 Count = EOS_LobbyDetails_GetMemberCount(LobbyDetailsHandle, &CountOptions);
+		CopyLobbyData(LobbyDetails, LobbyDetailsInfo, SearchResult.Session, Callback);
 
-		for (int32 Index = 0; Index < Count; Index++)
-		{
-			EOS_LobbyDetails_GetMemberByIndexOptions GetMemberByIndexOptions = { };
-			GetMemberByIndexOptions.ApiVersion = EOS_LOBBYDETAILS_GETMEMBERBYINDEX_API_LATEST;
-			GetMemberByIndexOptions.MemberIndex = Index;
-
-			EOS_ProductUserId TargetUserId = EOS_LobbyDetails_GetMemberByIndex(LobbyDetailsHandle, &GetMemberByIndexOptions);
-
-			EOSSubsystem->UserManager->GetEpicAccountIdAsync(TargetUserId, [this, LobbyDetailsHandle, SearchSettings, Position](const EOS_ProductUserId& ProductUserId, EOS_EpicAccountId& EpicAccountId) 
-				{
-					FUniqueNetIdRef UniqueNetId = FUniqueNetIdEOS::Create(EpicAccountId, ProductUserId);
-
-					FOnlineSessionSearchResult& SearchResult = SearchSettings->SearchResults[Position];
-					if (!SearchResult.Session.SessionSettings.MemberSettings.Contains(UniqueNetId))
-					{
-						SearchResult.Session.SessionSettings.MemberSettings.Add(UniqueNetId, FSessionSettings());
-					}
-
-					CopyLobbyMemberAttributes(LobbyDetailsHandle, ProductUserId, SearchResult.Session.SessionSettings.MemberSettings[UniqueNetId]);
-				});
-		}
-
-		LobbySearchResultsCache.Add(FString(LobbyDetailsInfo->LobbyId), MakeShared<EOS_HLobbyDetails>(LobbyDetailsHandle));
+		LobbySearchResultsCache.Add(FString(LobbyDetailsInfo->LobbyId), LobbyDetails);
 
 		EOS_LobbyDetails_Info_Release(LobbyDetailsInfo);
 
@@ -4011,11 +4021,11 @@ void FOnlineSessionEOS::AddLobbySearchResult(EOS_HLobbyDetails LobbyDetailsHandl
 	{
 		UE_LOG_ONLINE_SESSION(Warning, TEXT("[FOnlineSessionEOS::AddLobbySearchResult] LobbyDetails_CopyInfo not successful. Finished with EOS_EResult %s"), ANSI_TO_TCHAR(EOS_EResult_ToString(CopyResult)));
 
-		EOS_LobbyDetails_Release(LobbyDetailsHandle);
+		Callback(false);
 	}
 }
 
-void FOnlineSessionEOS::CopyLobbyData(EOS_HLobbyDetails LobbyDetailsHandle, EOS_LobbyDetails_Info* LobbyDetailsInfo, FOnlineSession& OutSession)
+void FOnlineSessionEOS::CopyLobbyData(const TSharedRef<FLobbyDetailsEOS>& LobbyDetails, EOS_LobbyDetails_Info* LobbyDetailsInfo, FOnlineSession& OutSession, const FOnCopyLobbyDataCompleteCallback& Callback)
 {
 	OutSession.SessionSettings.bUseLobbiesIfAvailable = true;
 	OutSession.SessionSettings.bIsLANMatch = false;
@@ -4047,16 +4057,50 @@ void FOnlineSessionEOS::CopyLobbyData(EOS_HLobbyDetails LobbyDetailsHandle, EOS_
 
 	OutSession.SessionSettings.bAllowInvites = (bool)LobbyDetailsInfo->bAllowInvites;
 
-	CopyLobbyAttributes(LobbyDetailsHandle, OutSession);
+	// We copy the settings related to lobby attributes
+	CopyLobbyAttributes(LobbyDetails, OutSession);
+
+	// Then we copy the settings for all lobby members
+	EOS_LobbyDetails_GetMemberCountOptions CountOptions = { };
+	CountOptions.ApiVersion = EOS_LOBBYDETAILS_GETMEMBERCOUNT_API_LATEST;
+	int32 Count = EOS_LobbyDetails_GetMemberCount(LobbyDetails->LobbyDetailsHandle, &CountOptions);
+
+	TArray<EOS_ProductUserId> TargetUserIds;
+	TargetUserIds.Reserve(Count);
+	for (int32 Index = 0; Index < Count; Index++)
+	{
+		EOS_LobbyDetails_GetMemberByIndexOptions GetMemberByIndexOptions = { };
+		GetMemberByIndexOptions.ApiVersion = EOS_LOBBYDETAILS_GETMEMBERBYINDEX_API_LATEST;
+		GetMemberByIndexOptions.MemberIndex = Index;
+
+		EOS_ProductUserId TargetUserId = EOS_LobbyDetails_GetMemberByIndex(LobbyDetails->LobbyDetailsHandle, &GetMemberByIndexOptions);
+
+		TargetUserIds.Add(TargetUserId);
+	}
+
+	EOSSubsystem->UserManager->ResolveUniqueNetIds(TargetUserIds, [this, LobbyDetails, &OutSession, OriginalCallback = Callback](TMap<EOS_ProductUserId, FUniqueNetIdEOSRef> ResolvedUniqueNetIds)
+	{
+		for (TMap<EOS_ProductUserId, FUniqueNetIdEOSRef>::TConstIterator It(ResolvedUniqueNetIds); It; ++It)
+		{
+			if (!OutSession.SessionSettings.MemberSettings.Contains(It.Value()))
+			{
+				OutSession.SessionSettings.MemberSettings.Add(It.Value(), FSessionSettings());
+			}
+
+			CopyLobbyMemberAttributes(*LobbyDetails, It.Key(), OutSession.SessionSettings.MemberSettings[It.Value()]);
+		}
+
+		OriginalCallback(true);
+	});
 }
 
-void FOnlineSessionEOS::CopyLobbyAttributes(EOS_HLobbyDetails LobbyDetailsHandle, FOnlineSession& OutSession)
+void FOnlineSessionEOS::CopyLobbyAttributes(const TSharedRef<FLobbyDetailsEOS>& LobbyDetails, FOnlineSession& OutSession)
 {
 	// In this method we are updating/adding attributes, but not removing
 
 	EOS_LobbyDetails_GetAttributeCountOptions CountOptions = { };
 	CountOptions.ApiVersion = EOS_LOBBYDETAILS_GETATTRIBUTECOUNT_API_LATEST;
-	int32 Count = EOS_LobbyDetails_GetAttributeCount(LobbyDetailsHandle, &CountOptions);
+	int32 Count = EOS_LobbyDetails_GetAttributeCount(LobbyDetails->LobbyDetailsHandle, &CountOptions);
 
 	for (int32 Index = 0; Index < Count; Index++)
 	{
@@ -4065,7 +4109,7 @@ void FOnlineSessionEOS::CopyLobbyAttributes(EOS_HLobbyDetails LobbyDetailsHandle
 		AttrOptions.AttrIndex = Index;
 
 		EOS_Lobby_Attribute* Attribute = NULL;
-		EOS_EResult ResultCode = EOS_LobbyDetails_CopyAttributeByIndex(LobbyDetailsHandle, &AttrOptions, &Attribute);
+		EOS_EResult ResultCode = EOS_LobbyDetails_CopyAttributeByIndex(LobbyDetails->LobbyDetailsHandle, &AttrOptions, &Attribute);
 		if (ResultCode == EOS_EResult::EOS_Success)
 		{
 			FString Key = UTF8_TO_TCHAR(Attribute->Data->Key);
@@ -4137,7 +4181,7 @@ void FOnlineSessionEOS::CopyLobbyAttributes(EOS_HLobbyDetails LobbyDetailsHandle
 	}
 }
 
-void FOnlineSessionEOS::CopyLobbyMemberAttributes(EOS_HLobbyDetails LobbyDetailsHandle, const EOS_ProductUserId& TargetUserId, FSessionSettings& OutSessionSettings)
+void FOnlineSessionEOS::CopyLobbyMemberAttributes(const FLobbyDetailsEOS& LobbyDetails, const EOS_ProductUserId& TargetUserId, FSessionSettings& OutSessionSettings)
 {
 	// In this method we are updating/adding attributes, but not removing
 
@@ -4145,7 +4189,7 @@ void FOnlineSessionEOS::CopyLobbyMemberAttributes(EOS_HLobbyDetails LobbyDetails
 	GetMemberAttributeCountOptions.ApiVersion = EOS_LOBBYDETAILS_GETMEMBERATTRIBUTECOUNT_API_LATEST;
 	GetMemberAttributeCountOptions.TargetUserId = TargetUserId;
 
-	uint32_t MemberAttributeCount = EOS_LobbyDetails_GetMemberAttributeCount(LobbyDetailsHandle, &GetMemberAttributeCountOptions);
+	uint32_t MemberAttributeCount = EOS_LobbyDetails_GetMemberAttributeCount(LobbyDetails.LobbyDetailsHandle, &GetMemberAttributeCountOptions);
 	for (uint32_t MemberAttributeIndex = 0; MemberAttributeIndex < MemberAttributeCount; MemberAttributeIndex++)
 	{
 		EOS_LobbyDetails_CopyMemberAttributeByIndexOptions AttrOptions = { };
@@ -4153,7 +4197,7 @@ void FOnlineSessionEOS::CopyLobbyMemberAttributes(EOS_HLobbyDetails LobbyDetails
 		AttrOptions.AttrIndex = MemberAttributeIndex;
 
 		EOS_Lobby_Attribute* Attribute = NULL;
-		EOS_EResult ResultCode = EOS_LobbyDetails_CopyMemberAttributeByIndex(LobbyDetailsHandle, &AttrOptions, &Attribute);
+		EOS_EResult ResultCode = EOS_LobbyDetails_CopyMemberAttributeByIndex(LobbyDetails.LobbyDetailsHandle, &AttrOptions, &Attribute);
 		if (ResultCode == EOS_EResult::EOS_Success)
 		{
 			FString Key = Attribute->Data->Key;
diff --git a/Engine/Plugins/Online/OnlineSubsystemEOS/Source/OnlineSubsystemEOS/Private/OnlineSessionEOS.h b/Engine/Plugins/Online/OnlineSubsystemEOS/Source/OnlineSubsystemEOS/Private/OnlineSessionEOS.h
--- a/Engine/Plugins/Online/OnlineSubsystemEOS/Source/OnlineSubsystemEOS/Private/OnlineSessionEOS.h
+++ b/Engine/Plugins/Online/OnlineSubsystemEOS/Source/OnlineSubsystemEOS/Private/OnlineSessionEOS.h
@@ -14,6 +14,10 @@ class FOnlineSubsystemEOS;
 	#include "eos_sessions_types.h"
 	#include "eos_lobby_types.h"
 
+static FName EOS_SESSION_ID = TEXT("EOS_SESSION_ID");
+static FName EOS_LOBBY_ID = TEXT("EOS_LOBBY_ID");
+TEMP_UNIQUENETIDSTRING_SUBCLASS(FUniqueNetIdEOSSession, EOS_SESSION_ID);
+TEMP_UNIQUENETIDSTRING_SUBCLASS(FUniqueNetIdEOSLobby, EOS_LOBBY_ID);
 
 struct FSessionSearchEOS
 {
@@ -30,6 +34,21 @@ struct FSessionSearchEOS
 	}
 };
 
+struct FLobbyDetailsEOS : FNoncopyable
+{
+	EOS_HLobbyDetails LobbyDetailsHandle;
+
+	FLobbyDetailsEOS(EOS_HLobbyDetails InLobbyDetailsHandle)
+		: LobbyDetailsHandle(InLobbyDetailsHandle)
+	{
+	}
+
+	virtual ~FLobbyDetailsEOS()
+	{
+		EOS_LobbyDetails_Release(LobbyDetailsHandle);
+	}
+};
+
 /**
  * Interface for interacting with EOS sessions
  */
@@ -175,6 +194,10 @@ public:
 
 	void CheckPendingSessionInvite();
 
+	void UpdateOrAddLobbyMember(const FUniqueNetIdEOSLobbyRef& LobbyNetId, const FUniqueNetIdEOSRef& PlayerId);
+	bool AddOnlineSessionMember(FName SessionName, const FUniqueNetIdRef& PlayerId);
+	bool RemoveOnlineSessionMember(FName SessionName, const FUniqueNetIdRef& PlayerId);
+
 	void RegisterLocalPlayers(class FNamedOnlineSession* Session);
 
 	void Init(const FString& InBucketId);
@@ -183,7 +206,8 @@ private:
 	// EOS Lobbies
 
 	EOS_HLobby LobbyHandle;
-	TMap<FString, TSharedRef<EOS_HLobbyDetails>> LobbySearchResultsCache;
+	TArray<TSharedRef<FLobbyDetailsEOS>> PendingLobbySearchResults;
+	TMap<FString, TSharedRef<FLobbyDetailsEOS>> LobbySearchResultsCache;
 
 	// Lobby session callbacks and methods
 	FCallbackBase* LobbyCreatedCallback;
@@ -221,7 +245,7 @@ private:
 	void OnLobbyInviteAccepted(const char* InviteId, const EOS_ProductUserId& LocalUserId, const EOS_ProductUserId& TargetUserId);
 	void OnJoinLobbyAccepted(const EOS_ProductUserId& LocalUserId, const EOS_UI_EventId& UiEventId);
 
-	// Lobby Update
+	// Methods to update an API Lobby from an OSS Lobby
 	void SetLobbyPermissionLevel(EOS_HLobbyModification LobbyModificationHandle, FNamedOnlineSession* Session);
 	void SetLobbyMaxMembers(EOS_HLobbyModification LobbyModificationHandle, FNamedOnlineSession* Session);
 	void SetLobbyAttributes(EOS_HLobbyModification LobbyModificationHandle, FNamedOnlineSession* Session);
@@ -229,19 +253,22 @@ private:
 	void AddLobbyMemberAttribute(EOS_HLobbyModification LobbyModificationHandle, const EOS_Lobby_AttributeData* Attribute);
 	void AddLobbyMember(const FUniqueNetIdStringRef LobbyNetId, const EOS_ProductUserId& TargetUserId);
 
+	// Methods to update an OSS Lobby from an API Lobby
+	typedef TFunction<void(bool bWasSuccessful)> FOnCopyLobbyDataCompleteCallback;
+	void CopyLobbyData(const TSharedRef<FLobbyDetailsEOS>& LobbyDetails, EOS_LobbyDetails_Info* LobbyDetailsInfo, FOnlineSession& OutSession, const FOnCopyLobbyDataCompleteCallback& Callback);
+	void CopyLobbyAttributes(const TSharedRef<FLobbyDetailsEOS>& LobbyDetails, FOnlineSession& OutSession);
+	void CopyLobbyMemberAttributes(const FLobbyDetailsEOS& LobbyDetails, const EOS_ProductUserId& TargetUserId, FSessionSettings& OutSessionSettings);
+
 	// Lobby search
 	void AddLobbySearchAttribute(EOS_HLobbySearch LobbySearchHandle, const EOS_Lobby_AttributeData* Attribute, EOS_EOnlineComparisonOp ComparisonOp);
-	void AddLobbySearchResult(EOS_HLobbyDetails LobbyDetailsHandle, const TSharedRef<FOnlineSessionSearch>& SearchSettings);
-	void CopyLobbyData(EOS_HLobbyDetails LobbyDetailsHandle, EOS_LobbyDetails_Info* LobbyDetailsInfo, FOnlineSession& OutSession);
-	void CopyLobbyAttributes(EOS_HLobbyDetails LobbyDetailsHandle, FOnlineSession& OutSession);
-	void CopyLobbyMemberAttributes(EOS_HLobbyDetails LobbyDetailsHandle, const EOS_ProductUserId& TargetUserId, FSessionSettings& OutSessionSettings);
+	void AddLobbySearchResult(const TSharedRef<FLobbyDetailsEOS>& LobbyDetails, const TSharedRef<FOnlineSessionSearch>& SearchSettings, const FOnCopyLobbyDataCompleteCallback& Callback);
 
 	// Helper methods
 	typedef TFunction<void(const EOS_ProductUserId& ProductUserId, EOS_EpicAccountId& EpicAccountId)> GetEpicAccountIdAsyncCallback;
 
 	void GetEpicAccountIdAsync(const EOS_ProductUserId& ProductUserId, const GetEpicAccountIdAsyncCallback& Callback);
 	void RegisterLobbyNotifications();
-	FNamedOnlineSession* GetNamedSessionFromLobbyId(const FUniqueNetIdString& LobbyId);
+	FNamedOnlineSession* GetNamedSessionFromLobbyId(const FUniqueNetIdEOSLobby& LobbyId);
 	bool GetEpicAccountIdFromProductUserId(const EOS_ProductUserId& ProductUserId, EOS_EpicAccountId& EpicAccountId);
 	EOS_ELobbyPermissionLevel GetLobbyPermissionLevelFromSessionSettings(const FOnlineSessionSettings& SessionSettings);
 	uint32_t GetLobbyMaxMembersFromSessionSettings(const FOnlineSessionSettings& SessionSettings);
diff --git a/Engine/Plugins/Online/OnlineSubsystemEOS/Source/OnlineSubsystemEOS/Private/UserManagerEOS.cpp b/Engine/Plugins/Online/OnlineSubsystemEOS/Source/OnlineSubsystemEOS/Private/UserManagerEOS.cpp
--- a/Engine/Plugins/Online/OnlineSubsystemEOS/Source/OnlineSubsystemEOS/Private/UserManagerEOS.cpp
+++ b/Engine/Plugins/Online/OnlineSubsystemEOS/Source/OnlineSubsystemEOS/Private/UserManagerEOS.cpp
@@ -1268,44 +1268,76 @@ bool FUserManagerEOS::GetEpicAccountIdFromProductUserId(const EOS_ProductUserId&
 	return bResult;
 }
 
-void FUserManagerEOS::GetEpicAccountIdAsync(const EOS_ProductUserId& ProductUserId, const GetEpicAccountIdAsyncCallback& Callback) const
+void FUserManagerEOS::ResolveUniqueNetId(const EOS_ProductUserId& ProductUserId, const FResolveUniqueNetIdCallback& Callback) const
 {
-	// We check first if the Product User Id has already been queried, which would allow us to retrieve its Epic Account Id directly
-	EOS_EpicAccountId AccountId;
-	if (GetEpicAccountIdFromProductUserId(ProductUserId, AccountId))
+	TArray<EOS_ProductUserId> ProductUserIds = { const_cast<EOS_ProductUserId>(ProductUserId) };
+
+	FResolveUniqueNetIdsCallback GroupCallback = [ProductUserId, OriginalCallback = Callback](TMap<EOS_ProductUserId, FUniqueNetIdEOSRef> ResolvedUniqueNetIds) {
+		OriginalCallback(ResolvedUniqueNetIds[ProductUserId]);
+	};
+
+	ResolveUniqueNetIds(ProductUserIds, GroupCallback);
+}
+
+void FUserManagerEOS::ResolveUniqueNetIds(const TArray<EOS_ProductUserId>& ProductUserIds, const FResolveUniqueNetIdsCallback& Callback) const
+{
+	TMap<EOS_ProductUserId, FUniqueNetIdEOSRef> ResolvedUniqueNetIds;
+	TArray<EOS_ProductUserId> ProductUserIdsToResolve;
+
+	for (const EOS_ProductUserId& ProductUserId : ProductUserIds)
 	{
-		Callback(ProductUserId, AccountId);
+		EOS_EpicAccountId EpicAccountId;
+
+		// We check first if the Product User Id has already been queried, which would allow us to retrieve its Epic Account Id directly
+		if (GetEpicAccountIdFromProductUserId(ProductUserId, EpicAccountId))
+		{
+			const FUniqueNetIdEOSRef UniqueNetId = FUniqueNetIdEOS::Create(MakeNetIdStringFromIds(EpicAccountId, ProductUserId));
+
+			ResolvedUniqueNetIds.Add(ProductUserId, UniqueNetId);
+		}
+		else
+		{
+			// If that's not the case, we'll have to query them first
+			ProductUserIdsToResolve.Add(ProductUserId);
+		}
 	}
-	else
-	{
-		// If it's the first time we want the Epic Account Id for this Product User Id, we have to query it first
-		TArray<EOS_ProductUserId> ProductUserIdList = { const_cast<EOS_ProductUserId>(ProductUserId) };
 
+	if (!ProductUserIdsToResolve.IsEmpty())
+	{
 		EOS_Connect_QueryProductUserIdMappingsOptions QueryProductUserIdMappingsOptions = {};
 		QueryProductUserIdMappingsOptions.ApiVersion = EOS_CONNECT_QUERYPRODUCTUSERIDMAPPINGS_API_LATEST;
 		QueryProductUserIdMappingsOptions.LocalUserId = EOSSubsystem->UserManager->GetLocalProductUserId(0);
-		QueryProductUserIdMappingsOptions.ProductUserIds = ProductUserIdList.GetData();
-		QueryProductUserIdMappingsOptions.ProductUserIdCount = 1;
+		QueryProductUserIdMappingsOptions.ProductUserIds = ProductUserIdsToResolve.GetData();
+		QueryProductUserIdMappingsOptions.ProductUserIdCount = ProductUserIdsToResolve.Num();
 
 		FConnectQueryProductUserIdMappingsCallback* CallbackObj = new FConnectQueryProductUserIdMappingsCallback(FUserManagerEOSConstWeakPtr(AsShared()));
-		CallbackObj->CallbackLambda = [this, ProductUserId, Callback](const EOS_Connect_QueryProductUserIdMappingsCallbackInfo* Data)
+		CallbackObj->CallbackLambda = [this, ProductUserIdsToResolve, ResolvedUniqueNetIds = MoveTemp(ResolvedUniqueNetIds), Callback](const EOS_Connect_QueryProductUserIdMappingsCallbackInfo* Data) mutable
 		{
-			EOS_EpicAccountId AccountId = nullptr;
-
-			if (Data->ResultCode == EOS_EResult::EOS_Success)
+			if (Data->ResultCode != EOS_EResult::EOS_Success)
 			{
-				GetEpicAccountIdFromProductUserId(ProductUserId, AccountId);
+				UE_LOG_ONLINE(Verbose, TEXT("[FUserManagerEOS::ResolveUniqueNetIds] EOS_Connect_QueryProductUserIdMappings not successful for user (%s). Finished with EOS_EResult %s."), *LexToString(Data->LocalUserId), ANSI_TO_TCHAR(EOS_EResult_ToString(Data->ResultCode)));
 			}
-			else
+
+			for (const EOS_ProductUserId& ProductUserId : ProductUserIdsToResolve)
 			{
-				UE_LOG_ONLINE(Verbose, TEXT("[FUserManagerEOS::GetEpicAccountIdAsync] EOS_Connect_QueryProductUserIdMappings not successful for ProductUserId (%s). Finished with EOS_EResult %s."), *LexToString(ProductUserId), ANSI_TO_TCHAR(EOS_EResult_ToString(Data->ResultCode)));
+				EOS_EpicAccountId EpicAccountId = nullptr;
+
+				GetEpicAccountIdFromProductUserId(ProductUserId, EpicAccountId);
+
+				const FUniqueNetIdEOSRef UniqueNetId = FUniqueNetIdEOS::Create(MakeNetIdStringFromIds(EpicAccountId, ProductUserId));
+
+				ResolvedUniqueNetIds.Add(ProductUserId, UniqueNetId);
 			}
 
-			Callback(ProductUserId, AccountId);
+			Callback(ResolvedUniqueNetIds);
 		};
 
 		EOS_Connect_QueryProductUserIdMappings(EOSSubsystem->ConnectHandle, &QueryProductUserIdMappingsOptions, CallbackObj, CallbackObj->GetCallbackPtr());
 	}
+	else
+	{
+		Callback(ResolvedUniqueNetIds);
+	}
 }
 
 FOnlineUserPtr FUserManagerEOS::GetLocalOnlineUser(int32 LocalUserNum) const
diff --git a/Engine/Plugins/Online/OnlineSubsystemEOS/Source/OnlineSubsystemEOS/Private/UserManagerEOS.h b/Engine/Plugins/Online/OnlineSubsystemEOS/Source/OnlineSubsystemEOS/Private/UserManagerEOS.h
--- a/Engine/Plugins/Online/OnlineSubsystemEOS/Source/OnlineSubsystemEOS/Private/UserManagerEOS.h
+++ b/Engine/Plugins/Online/OnlineSubsystemEOS/Source/OnlineSubsystemEOS/Private/UserManagerEOS.h
@@ -354,9 +354,11 @@ public:
 	int32 GetLocalUserNumFromUniqueNetId(const FUniqueNetId& NetId) const;
 	bool IsLocalUser(const FUniqueNetId& NetId) const;
 
-	typedef TFunction<void(const EOS_ProductUserId& ProductUserId, EOS_EpicAccountId& EpicAccountId)> GetEpicAccountIdAsyncCallback;
+	typedef TFunction<void(TMap<EOS_ProductUserId, FUniqueNetIdEOSRef> ResolvedUniqueNetIds)> FResolveUniqueNetIdsCallback;
+	typedef TFunction<void(FUniqueNetIdEOSRef ResolvedUniqueNetId)> FResolveUniqueNetIdCallback;
 	bool GetEpicAccountIdFromProductUserId(const EOS_ProductUserId& ProductUserId, EOS_EpicAccountId& OutEpicAccountId) const;
-	void GetEpicAccountIdAsync(const EOS_ProductUserId& ProductUserId, const GetEpicAccountIdAsyncCallback& Callback) const;
+	void ResolveUniqueNetId(const EOS_ProductUserId& ProductUserId, const FResolveUniqueNetIdCallback& Callback) const;
+	void ResolveUniqueNetIds(const TArray<EOS_ProductUserId>& ProductUserIds, const FResolveUniqueNetIdsCallback& Callback) const;
 
 	FOnlineUserPtr GetLocalOnlineUser(int32 LocalUserNum) const;
 	FOnlineUserPtr GetOnlineUser(EOS_ProductUserId UserId) const;
