- Fix broken console builds, they build against console specific SDK headers which don't contain this new enum

diff --git a/Engine/Plugins/Online/OnlineSubsystemEOS/Source/OnlineSubsystemEOS/Private/UserManagerEOS.cpp b/Engine/Plugins/Online/OnlineSubsystemEOS/Source/OnlineSubsystemEOS/Private/UserManagerEOS.cpp
--- a/Engine/Plugins/Online/OnlineSubsystemEOS/Source/OnlineSubsystemEOS/Private/UserManagerEOS.cpp
+++ b/Engine/Plugins/Online/OnlineSubsystemEOS/Source/OnlineSubsystemEOS/Private/UserManagerEOS.cpp
@@ -93,11 +93,13 @@ static inline EOS_Presence_EStatus ToEOS_Presence_EStatus(EOnlinePresenceState::
 
 static inline EOS_EExternalCredentialType ToEOS_EExternalCredentialType(FName OSSName, const FOnlineAccountCredentials& AccountCredentials)
 {
+#if PLATFORM_DESKTOP
 	if (OSSName == STEAM_SUBSYSTEM)
 	{
 		return EOS_EExternalCredentialType::EOS_ECT_STEAM_SESSION_TICKET;
 	}
-	else if (OSSName == PS4_SUBSYSTEM || USE_PSN_ID_TOKEN)
+#endif
+	if (OSSName == PS4_SUBSYSTEM || USE_PSN_ID_TOKEN)
 	{
 		return EOS_EExternalCredentialType::EOS_ECT_PSN_ID_TOKEN;
 	}
