Add support for account portal logins using the EOS overlay

diff --git a/Engine/Plugins/Online/OnlineSubsystemEOS/Source/OnlineSubsystemEOS/Private/EOSSettings.cpp b/Engine/Plugins/Online/OnlineSubsystemEOS/Source/OnlineSubsystemEOS/Private/EOSSettings.cpp
--- a/Engine/Plugins/Online/OnlineSubsystemEOS/Source/OnlineSubsystemEOS/Private/EOSSettings.cpp
+++ b/Engine/Plugins/Online/OnlineSubsystemEOS/Source/OnlineSubsystemEOS/Private/EOSSettings.cpp
@@ -151,6 +151,7 @@ const FEOSSettings& UEOSSettings::ManualGetSettings()
 		GConfig->GetInt(INI_SECTION, TEXT("TitleStorageReadChunkLength"), CachedSettings->TitleStorageReadChunkLength, GEngineIni);
 		GConfig->GetBool(INI_SECTION, TEXT("bEnableOverlay"), CachedSettings->bEnableOverlay, GEngineIni);
 		GConfig->GetBool(INI_SECTION, TEXT("bEnableSocialOverlay"), CachedSettings->bEnableSocialOverlay, GEngineIni);
+		GConfig->GetBool(INI_SECTION, TEXT("bEnableEditorOverlay"), CachedSettings->bEnableEditorOverlay, GEngineIni);
 		GConfig->GetBool(INI_SECTION, TEXT("bShouldEnforceBeingLaunchedByEGS"), CachedSettings->bShouldEnforceBeingLaunchedByEGS, GEngineIni);
 		GConfig->GetBool(INI_SECTION, TEXT("bUseEAS"), CachedSettings->bUseEAS, GEngineIni);
 		GConfig->GetBool(INI_SECTION, TEXT("bUseEOSConnect"), CachedSettings->bUseEOSConnect, GEngineIni);
@@ -175,6 +176,7 @@ FEOSSettings UEOSSettings::ToNative() const
 	Native.TitleStorageReadChunkLength = TitleStorageReadChunkLength;
 	Native.bEnableOverlay = bEnableOverlay;
 	Native.bEnableSocialOverlay = bEnableSocialOverlay;
+	Native.bEnableEditorOverlay = bEnableEditorOverlay;
 	Native.bShouldEnforceBeingLaunchedByEGS = bShouldEnforceBeingLaunchedByEGS;
 	Native.bUseEAS = bUseEAS;
 	Native.bUseEOSConnect = bUseEOSConnect;
@@ -298,6 +300,7 @@ void UEOSSettings::PostEditChangeProperty(FPropertyChangedEvent& PropertyChanged
 		if (!bEnableOverlay)
 		{
 			bEnableSocialOverlay = false;
+			bEnableEditorOverlay = false;
 		}
 	}
 
diff --git a/Engine/Plugins/Online/OnlineSubsystemEOS/Source/OnlineSubsystemEOS/Private/OnlineSubsystemEOS.cpp b/Engine/Plugins/Online/OnlineSubsystemEOS/Source/OnlineSubsystemEOS/Private/OnlineSubsystemEOS.cpp
--- a/Engine/Plugins/Online/OnlineSubsystemEOS/Source/OnlineSubsystemEOS/Private/OnlineSubsystemEOS.cpp
+++ b/Engine/Plugins/Online/OnlineSubsystemEOS/Source/OnlineSubsystemEOS/Private/OnlineSubsystemEOS.cpp
@@ -224,7 +224,12 @@ bool FOnlineSubsystemEOS::PlatformCreate()
 	{
 		OverlayFlags |= EOS_PF_DISABLE_SOCIAL_OVERLAY;
 	}
-	PlatformOptions.Flags = IsRunningGame() ? OverlayFlags : EOS_PF_DISABLE_OVERLAY;
+
+	// Don't allow the overlay to be used in the editor when running PIE.
+	const bool bEditorOverlayAllowed = EOSSettings.bEnableEditorOverlay && InstanceName == FOnlineSubsystemImpl::DefaultInstanceName;
+	const bool bOverlayAllowed = IsRunningGame() || bEditorOverlayAllowed;
+
+	PlatformOptions.Flags = bOverlayAllowed ? OverlayFlags : EOS_PF_DISABLE_OVERLAY;
 	// Make the cache directory be in the user's writable area
 
 	const FString CacheDir = EOSSDKManager->GetCacheDirBase() / ArtifactName / EOSSettings.CacheDir;
diff --git a/Engine/Plugins/Online/OnlineSubsystemEOS/Source/OnlineSubsystemEOS/Public/EOSSettings.h b/Engine/Plugins/Online/OnlineSubsystemEOS/Source/OnlineSubsystemEOS/Public/EOSSettings.h
--- a/Engine/Plugins/Online/OnlineSubsystemEOS/Source/OnlineSubsystemEOS/Public/EOSSettings.h
+++ b/Engine/Plugins/Online/OnlineSubsystemEOS/Source/OnlineSubsystemEOS/Public/EOSSettings.h
@@ -73,6 +73,7 @@ struct FEOSSettings
 	int32 TitleStorageReadChunkLength;
 	bool bEnableOverlay;
 	bool bEnableSocialOverlay;
+	bool bEnableEditorOverlay;
 	bool bShouldEnforceBeingLaunchedByEGS;
 	bool bUseEAS;
 	bool bUseEOSConnect;
@@ -115,6 +116,10 @@ public:
 	UPROPERTY(Config, EditAnywhere, BlueprintReadOnly, Category="EOS Settings")
 	bool bEnableSocialOverlay = false;
 
+	/** Set to true to enable the overlay when running in the editor */
+	UPROPERTY(Config, EditAnywhere, BlueprintReadOnly, Category = "EOS Settings")
+	bool bEnableEditorOverlay = false;
+
 	/** Set to true to enable the social overlay (friends, invites, etc.) */
 	UPROPERTY(Config, EditAnywhere, BlueprintReadOnly, Category="EOS Settings", DisplayName="Require Being Launched by the Epic Games Store")
 	bool bShouldEnforceBeingLaunchedByEGS = false;
