Fix CacheDirectory issue with Android EOSSDK initialization
Move support for accountportal login to the base FEOSHelpers class for subclasses to use.  Use this on Android.

diff --git a/Engine/Plugins/Online/EOSShared/Source/EOSShared/Private/Android/AndroidEOSSDKManager.cpp b/Engine/Plugins/Online/EOSShared/Source/EOSShared/Private/Android/AndroidEOSSDKManager.cpp
--- a/Engine/Plugins/Online/EOSShared/Source/EOSShared/Private/Android/AndroidEOSSDKManager.cpp
+++ b/Engine/Plugins/Online/EOSShared/Source/EOSShared/Private/Android/AndroidEOSSDKManager.cpp
@@ -1,6 +1,7 @@
 // Copyright Epic Games, Inc. All Rights Reserved.
 
 #include "AndroidEOSSDKManager.h"
+#include "HAL/FileManager.h"
 
 #if WITH_EOS_SDK
 
@@ -20,4 +21,10 @@ EOS_EResult FAndroidEOSSDKManager::EOSInitialize(EOS_InitializeOptions& Options)
 	return FEOSSDKManager::EOSInitialize(Options);
 }
 
+FString FAndroidEOSSDKManager::GetCacheDirBase() const
+{
+	FString BaseCacheDirBase = FEOSSDKManager::GetCacheDirBase();
+	return IFileManager::Get().ConvertToAbsolutePathForExternalAppForWrite(*BaseCacheDirBase);
+};
+
 #endif // WITH_EOS_SDK
\ No newline at end of file
diff --git a/Engine/Plugins/Online/EOSShared/Source/EOSShared/Private/Android/AndroidEOSSDKManager.h b/Engine/Plugins/Online/EOSShared/Source/EOSShared/Private/Android/AndroidEOSSDKManager.h
--- a/Engine/Plugins/Online/EOSShared/Source/EOSShared/Private/Android/AndroidEOSSDKManager.h
+++ b/Engine/Plugins/Online/EOSShared/Source/EOSShared/Private/Android/AndroidEOSSDKManager.h
@@ -10,6 +10,7 @@ class FAndroidEOSSDKManager : public FEOSSDKManager
 {
 public:
 	virtual EOS_EResult EOSInitialize(EOS_InitializeOptions& Options) override;
+	virtual FString GetCacheDirBase() const override;
 };
 
 using FPlatformEOSSDKManager = FAndroidEOSSDKManager;
diff --git a/Engine/Plugins/Online/OnlineSubsystemEOS/Source/OnlineSubsystemEOS/Private/Android/AndroidEOSHelpers.cpp b/Engine/Plugins/Online/OnlineSubsystemEOS/Source/OnlineSubsystemEOS/Private/Android/AndroidEOSHelpers.cpp
new file mode 100644
--- /dev/null
+++ b/Engine/Plugins/Online/OnlineSubsystemEOS/Source/OnlineSubsystemEOS/Private/Android/AndroidEOSHelpers.cpp
@@ -0,0 +1,12 @@
+// Copyright Epic Games, Inc. All Rights Reserved.
+
+#include "AndroidEOSHelpers.h"
+
+#if WITH_EOS_SDK
+
+void FAndroidEOSHelpers::PlatformTriggerLoginUI(FOnlineSubsystemEOS* InEOSSubsystem, const int ControllerIndex, bool bShowOnlineOnly, bool bShowSkipButton, const FOnLoginUIClosedDelegate& Delegate)
+{
+	ShowAccountPortalUI(InEOSSubsystem, ControllerIndex, Delegate);
+}
+
+#endif
diff --git a/Engine/Plugins/Online/OnlineSubsystemEOS/Source/OnlineSubsystemEOS/Private/Android/AndroidEOSHelpers.h b/Engine/Plugins/Online/OnlineSubsystemEOS/Source/OnlineSubsystemEOS/Private/Android/AndroidEOSHelpers.h
--- a/Engine/Plugins/Online/OnlineSubsystemEOS/Source/OnlineSubsystemEOS/Private/Android/AndroidEOSHelpers.h
+++ b/Engine/Plugins/Online/OnlineSubsystemEOS/Source/OnlineSubsystemEOS/Private/Android/AndroidEOSHelpers.h
@@ -6,6 +6,12 @@
 
 #include "EOSHelpers.h"
 
-using FPlatformEOSHelpers = FEOSHelpers;
+class FAndroidEOSHelpers : public FEOSHelpers
+{
+public:
+	virtual void PlatformTriggerLoginUI(FOnlineSubsystemEOS* EOSSubsystem, const int ControllerIndex, bool bShowOnlineOnly, bool bShowSkipButton, const FOnLoginUIClosedDelegate& Delegate) override;
+};
+
+using FPlatformEOSHelpers = FAndroidEOSHelpers;
 
 #endif
\ No newline at end of file
diff --git a/Engine/Plugins/Online/OnlineSubsystemEOS/Source/OnlineSubsystemEOS/Private/EOSHelpers.cpp b/Engine/Plugins/Online/OnlineSubsystemEOS/Source/OnlineSubsystemEOS/Private/EOSHelpers.cpp
--- a/Engine/Plugins/Online/OnlineSubsystemEOS/Source/OnlineSubsystemEOS/Private/EOSHelpers.cpp
+++ b/Engine/Plugins/Online/OnlineSubsystemEOS/Source/OnlineSubsystemEOS/Private/EOSHelpers.cpp
@@ -38,4 +38,25 @@ IEOSPlatformHandlePtr FEOSHelpers::CreatePlatform(EOS_Platform_Options& Platform
 	return nullptr;
 }
 
+void FEOSHelpers::ShowAccountPortalUI(FOnlineSubsystemEOS* InEOSSubsystem, const int ControllerIndex, const FOnLoginUIClosedDelegate& Delegate)
+{
+	check(InEOSSubsystem != nullptr);
+	FDelegateHandle* DelegateHandle = new FDelegateHandle;
+	*DelegateHandle = InEOSSubsystem->UserManager->AddOnLoginCompleteDelegate_Handle(ControllerIndex, FOnLoginCompleteDelegate::CreateRaw(this, &FEOSHelpers::OnAccountPortalLoginComplete, InEOSSubsystem, Delegate, DelegateHandle));
+
+	FOnlineAccountCredentials* Credentials = new FOnlineAccountCredentials(TEXT("accountportal"), TEXT(""), TEXT(""));
+	InEOSSubsystem->UserManager->Login(ControllerIndex, *Credentials);
+}
+
+void FEOSHelpers::OnAccountPortalLoginComplete(int ControllerIndex, bool bWasSuccessful, const FUniqueNetId& UserId, const FString& ErrorString, FOnlineSubsystemEOS* InEOSSubsystem, const FOnLoginUIClosedDelegate LoginUIClosedDelegate, FDelegateHandle* LoginDelegateHandle) const
+{
+	FOnlineError Error(bWasSuccessful);
+	Error.SetFromErrorCode(ErrorString);
+
+	InEOSSubsystem->UserManager->ClearOnLoginCompleteDelegate_Handle(ControllerIndex, *LoginDelegateHandle);
+	delete LoginDelegateHandle;
+
+	LoginUIClosedDelegate.ExecuteIfBound(UserId.AsShared(), ControllerIndex, Error);
+}
+
 #endif // WITH_EOS_SDK
\ No newline at end of file
diff --git a/Engine/Plugins/Online/OnlineSubsystemEOS/Source/OnlineSubsystemEOS/Private/EOSHelpers.h b/Engine/Plugins/Online/OnlineSubsystemEOS/Source/OnlineSubsystemEOS/Private/EOSHelpers.h
--- a/Engine/Plugins/Online/OnlineSubsystemEOS/Source/OnlineSubsystemEOS/Private/EOSHelpers.h
+++ b/Engine/Plugins/Online/OnlineSubsystemEOS/Source/OnlineSubsystemEOS/Private/EOSHelpers.h
@@ -23,6 +23,14 @@ public:
 	virtual void PlatformAuthCredentials(EOS_Auth_Credentials &Credentials);
 	virtual void PlatformTriggerLoginUI(FOnlineSubsystemEOS* EOSSubsystem, const int ControllerIndex, bool bShowOnlineOnly, bool bShowSkipButton, const FOnLoginUIClosedDelegate& Delegate);
 	virtual IEOSPlatformHandlePtr CreatePlatform(EOS_Platform_Options& PlatformOptions);
+
+protected:
+	/** Shared LoginUI logic that can be used by platforms that support LoginUI */
+	void ShowAccountPortalUI(FOnlineSubsystemEOS* InEOSSubsystem, const int ControllerIndex, const FOnLoginUIClosedDelegate& Delegate);
+
+private:
+	/** Completion handler for ShowAccountPortalUI */
+	void OnAccountPortalLoginComplete(int ControllerIndex, bool bWasSuccessful, const FUniqueNetId& UserId, const FString& ErrorString, FOnlineSubsystemEOS* InEOSSubsystem, const FOnLoginUIClosedDelegate LoginUIClosedDelegate, FDelegateHandle* LoginDelegateHandle) const;
 };
 
 #endif // WITH_EOS_SDK
\ No newline at end of file
diff --git a/Engine/Plugins/Online/OnlineSubsystemEOS/Source/OnlineSubsystemEOS/Private/IOS/IOSEOSHelpers.cpp b/Engine/Plugins/Online/OnlineSubsystemEOS/Source/OnlineSubsystemEOS/Private/IOS/IOSEOSHelpers.cpp
--- a/Engine/Plugins/Online/OnlineSubsystemEOS/Source/OnlineSubsystemEOS/Private/IOS/IOSEOSHelpers.cpp
+++ b/Engine/Plugins/Online/OnlineSubsystemEOS/Source/OnlineSubsystemEOS/Private/IOS/IOSEOSHelpers.cpp
@@ -74,26 +74,7 @@ void FIOSEOSHelpers::PlatformAuthCredentials(EOS_Auth_Credentials &Credentials)
 
 void FIOSEOSHelpers::PlatformTriggerLoginUI(FOnlineSubsystemEOS* InEOSSubsystem, const int ControllerIndex, bool bShowOnlineOnly, bool bShowSkipButton, const FOnLoginUIClosedDelegate& Delegate)
 {
-	EOSSubsystem = InEOSSubsystem;
-	check(EOSSubsystem != nullptr);
-
-	LoginUIClosedDelegate = Delegate;
-	LoginCompleteDelegate = EOSSubsystem->UserManager->AddOnLoginCompleteDelegate_Handle(ControllerIndex, FOnLoginCompleteDelegate::CreateRaw(this, &FIOSEOSHelpers::OnLoginComplete));
-
-	FOnlineAccountCredentials *Credentials = new FOnlineAccountCredentials(TEXT("accountportal"), TEXT(""), TEXT(""));
-	EOSSubsystem->UserManager->Login(ControllerIndex, *Credentials);
-}
-
-void FIOSEOSHelpers::OnLoginComplete(int ControllerIndex, bool bWasSuccessful, const FUniqueNetId& UserId, const FString& ErrorString)
-{
-	FOnlineError Error(bWasSuccessful);
-	Error.SetFromErrorCode(ErrorString);
-
-	check(EOSSubsystem != nullptr);
-
-	EOSSubsystem->UserManager->ClearOnLoginCompleteDelegate_Handle(ControllerIndex, LoginCompleteDelegate);
-
-	LoginUIClosedDelegate.ExecuteIfBound(UserId, ControllerIndex, Error);
+	ShowAccountPortalUI(InEOSSubsystem, ControllerIndex, Delegate);
 }
 
 #endif
diff --git a/Engine/Plugins/Online/OnlineSubsystemEOS/Source/OnlineSubsystemEOS/Private/IOS/IOSEOSHelpers.h b/Engine/Plugins/Online/OnlineSubsystemEOS/Source/OnlineSubsystemEOS/Private/IOS/IOSEOSHelpers.h
--- a/Engine/Plugins/Online/OnlineSubsystemEOS/Source/OnlineSubsystemEOS/Private/IOS/IOSEOSHelpers.h
+++ b/Engine/Plugins/Online/OnlineSubsystemEOS/Source/OnlineSubsystemEOS/Private/IOS/IOSEOSHelpers.h
@@ -11,7 +11,7 @@
 
 class FOnlineSubsystemEOS;
 
-class FIOSEOSHelpers : FEOSHelpers
+class FIOSEOSHelpers : public FEOSHelpers
 {
 public:
 	virtual ~FIOSEOSHelpers();
@@ -20,13 +20,7 @@ public:
 	virtual void PlatformTriggerLoginUI(FOnlineSubsystemEOS* EOSSubsystem, const int ControllerIndex, bool bShowOnlineOnly, bool bShowSkipButton, const FOnLoginUIClosedDelegate& Delegate) override;
 
 private:
-	void OnLoginComplete(int ControllerIndex, bool bWasSuccessful, const FUniqueNetId& UserId, const FString& ErrorString);
-
 	EOS_IOS_Auth_CredentialsOptions CredentialsOptions;
-
-	FOnlineSubsystemEOS* EOSSubsystem;
-	FOnLoginUIClosedDelegate LoginUIClosedDelegate;
-	FDelegateHandle LoginCompleteDelegate;
 };
 
 using FPlatformEOSHelpers = FIOSEOSHelpers;
