Refactor of session object captures in session methods to avoid errors in lambdas.

diff --git a/Engine/Plugins/Online/OnlineSubsystemEOS/Source/OnlineSubsystemEOS/Private/OnlineSessionEOS.cpp b/Engine/Plugins/Online/OnlineSubsystemEOS/Source/OnlineSubsystemEOS/Private/OnlineSessionEOS.cpp
--- a/Engine/Plugins/Online/OnlineSubsystemEOS/Source/OnlineSubsystemEOS/Private/OnlineSessionEOS.cpp
+++ b/Engine/Plugins/Online/OnlineSubsystemEOS/Source/OnlineSubsystemEOS/Private/OnlineSessionEOS.cpp
@@ -1322,16 +1322,26 @@ uint32 FOnlineSessionEOS::StartEOSSession(FNamedOnlineSession* Session)
 
 	FSessionStartOptions Options(TCHAR_TO_UTF8(*Session->SessionName.ToString()));
 	FStartSessionCallback* CallbackObj = new FStartSessionCallback(FOnlineSessionEOSWeakPtr(AsShared()));
-	CallbackObj->CallbackLambda = [this, Session](const EOS_Sessions_StartSessionCallbackInfo* Data)
+	CallbackObj->CallbackLambda = [this, SessionName = Session->SessionName](const EOS_Sessions_StartSessionCallbackInfo* Data)
 	{
-		Session->SessionState = EOnlineSessionState::InProgress;
+		bool bWasSuccessful = false;
 
-		bool bWasSuccessful = Data->ResultCode == EOS_EResult::EOS_Success;
-		if (!bWasSuccessful)
+		if (FNamedOnlineSession* Session = GetNamedSession(SessionName))
+		{
+			Session->SessionState = EOnlineSessionState::InProgress;
+
+			bWasSuccessful = Data->ResultCode == EOS_EResult::EOS_Success;
+			if (!bWasSuccessful)
+			{
+				UE_LOG_ONLINE_SESSION(Error, TEXT("EOS_Sessions_StartSession() failed with EOS result code (%s)"), ANSI_TO_TCHAR(EOS_EResult_ToString(Data->ResultCode)));
+			}
+		}
+		else
 		{
-			UE_LOG_ONLINE_SESSION(Error, TEXT("EOS_Sessions_StartSession() failed with EOS result code (%s)"), ANSI_TO_TCHAR(EOS_EResult_ToString(Data->ResultCode)));
+			UE_LOG_ONLINE_SESSION(Verbose, TEXT("Session [%s] not found"), *SessionName.ToString());
 		}
-		TriggerOnStartSessionCompleteDelegates(Session->SessionName, bWasSuccessful);
+
+		TriggerOnStartSessionCompleteDelegates(SessionName, bWasSuccessful);
 	};
 
 	EOS_Sessions_StartSession(EOSSubsystem->SessionsHandle, &Options, CallbackObj, CallbackObj->GetCallbackPtr());
@@ -1343,11 +1353,18 @@ uint32 FOnlineSessionEOS::StartLobbySession(FNamedOnlineSession* Session)
 {
 	Session->SessionState = EOnlineSessionState::Starting;
 
-	EOSSubsystem->ExecuteNextTick([this, Session]()
+	EOSSubsystem->ExecuteNextTick([this, SessionName = Session->SessionName]()
 	{
-		Session->SessionState = EOnlineSessionState::InProgress;
+		if (FNamedOnlineSession* Session = GetNamedSession(SessionName))
+		{
+			Session->SessionState = EOnlineSessionState::InProgress;
+		}
+		else
+		{
+			UE_LOG_ONLINE_SESSION(Verbose, TEXT("Session [%s] not found"), *SessionName.ToString());
+		}
 
-		TriggerOnStartSessionCompleteDelegates(Session->SessionName, true);
+		TriggerOnStartSessionCompleteDelegates(SessionName, true);
 	});
 
 	return ONLINE_IO_PENDING;
@@ -1423,15 +1440,25 @@ uint32 FOnlineSessionEOS::UpdateEOSSession(FNamedOnlineSession* Session)
 	}
 
 	FUpdateSessionCallback* CallbackObj = new FUpdateSessionCallback(FOnlineSessionEOSWeakPtr(AsShared()));
-	CallbackObj->CallbackLambda = [this, Session](const EOS_Sessions_UpdateSessionCallbackInfo* Data)
+	CallbackObj->CallbackLambda = [this, SessionName = Session->SessionName](const EOS_Sessions_UpdateSessionCallbackInfo* Data)
 	{
-		bool bWasSuccessful = Data->ResultCode == EOS_EResult::EOS_Success || Data->ResultCode == EOS_EResult::EOS_Sessions_OutOfSync;
-		if (!bWasSuccessful)
+		bool bWasSuccessful = false;
+		
+		if (FNamedOnlineSession* Session = GetNamedSession(SessionName))
 		{
-			Session->SessionState = EOnlineSessionState::NoSession;
-			UE_LOG_ONLINE_SESSION(Error, TEXT("EOS_Sessions_UpdateSession() failed with EOS result code (%s)"), ANSI_TO_TCHAR(EOS_EResult_ToString(Data->ResultCode)));
+			bWasSuccessful = Data->ResultCode == EOS_EResult::EOS_Success || Data->ResultCode == EOS_EResult::EOS_Sessions_OutOfSync;
+			if (!bWasSuccessful)
+			{
+				Session->SessionState = EOnlineSessionState::NoSession;
+				UE_LOG_ONLINE_SESSION(Error, TEXT("EOS_Sessions_UpdateSession() failed with EOS result code (%s)"), ANSI_TO_TCHAR(EOS_EResult_ToString(Data->ResultCode)));
+			}
+		}
+		else
+		{
+			UE_LOG_ONLINE_SESSION(Verbose, TEXT("Session [%s] not found"), *SessionName.ToString());
 		}
-		TriggerOnUpdateSessionCompleteDelegates(Session->SessionName, bWasSuccessful);
+
+		TriggerOnUpdateSessionCompleteDelegates(SessionName, bWasSuccessful);
 	};
 
 	return SharedSessionUpdate(SessionModHandle, Session, CallbackObj);
@@ -1526,16 +1553,26 @@ uint32 FOnlineSessionEOS::EndEOSSession(FNamedOnlineSession* Session)
 
 	FSessionEndOptions Options(TCHAR_TO_UTF8(*Session->SessionName.ToString()));
 	FEndSessionCallback* CallbackObj = new FEndSessionCallback(FOnlineSessionEOSWeakPtr(AsShared()));
-	CallbackObj->CallbackLambda = [this, Session](const EOS_Sessions_EndSessionCallbackInfo* Data)
+	CallbackObj->CallbackLambda = [this, SessionName = Session->SessionName](const EOS_Sessions_EndSessionCallbackInfo* Data)
 	{
-		Session->SessionState = EOnlineSessionState::Ended;
+		bool bWasSuccessful = false;
 
-		bool bWasSuccessful = Data->ResultCode == EOS_EResult::EOS_Success;
-		if (!bWasSuccessful)
+		if (FNamedOnlineSession* Session = GetNamedSession(SessionName))
 		{
-			UE_LOG_ONLINE_SESSION(Error, TEXT("EOS_Sessions_EndSession() failed with EOS result code (%s)"), ANSI_TO_TCHAR(EOS_EResult_ToString(Data->ResultCode)));
+			Session->SessionState = EOnlineSessionState::Ended;
+
+			bWasSuccessful = Data->ResultCode == EOS_EResult::EOS_Success;
+			if (!bWasSuccessful)
+			{
+				UE_LOG_ONLINE_SESSION(Error, TEXT("EOS_Sessions_EndSession() failed with EOS result code (%s)"), ANSI_TO_TCHAR(EOS_EResult_ToString(Data->ResultCode)));
+			}
 		}
-		TriggerOnEndSessionCompleteDelegates(Session->SessionName, bWasSuccessful);
+		else
+		{
+			UE_LOG_ONLINE_SESSION(Verbose, TEXT("Session [%s] not found"), *SessionName.ToString());
+		}
+
+		TriggerOnEndSessionCompleteDelegates(SessionName, bWasSuccessful);
 	};
 
 	EOS_Sessions_EndSession(EOSSubsystem->SessionsHandle, &Options, CallbackObj, CallbackObj->GetCallbackPtr());
@@ -1667,18 +1704,26 @@ uint32 FOnlineSessionEOS::DestroyEOSSession(FNamedOnlineSession* Session, const
 
 	FSessionDestroyOptions Options(TCHAR_TO_UTF8(*Session->SessionName.ToString()));
 	FDestroySessionCallback* CallbackObj = new FDestroySessionCallback(FOnlineSessionEOSWeakPtr(AsShared()));
-	CallbackObj->CallbackLambda = [this, Session](const EOS_Sessions_DestroySessionCallbackInfo* Data)
+	CallbackObj->CallbackLambda = [this, SessionName = Session->SessionName](const EOS_Sessions_DestroySessionCallbackInfo* Data)
 	{
 		EndSessionAnalytics();
 
-		Session->SessionState = EOnlineSessionState::NoSession;
-		FName SessionName = Session->SessionName;
+		bool bWasSuccessful = false;
+		if (FNamedOnlineSession* Session = GetNamedSession(SessionName))
+		{
+			Session->SessionState = EOnlineSessionState::NoSession;
 
-		bool bWasSuccessful = Data->ResultCode == EOS_EResult::EOS_Success;
-		if (!bWasSuccessful)
+			bWasSuccessful = Data->ResultCode == EOS_EResult::EOS_Success;
+			if (!bWasSuccessful)
+			{
+				UE_LOG_ONLINE_SESSION(Error, TEXT("EOS_Sessions_DestroySession() failed with EOS result code (%s)"), ANSI_TO_TCHAR(EOS_EResult_ToString(Data->ResultCode)));
+			}
+		}
+		else
 		{
-			UE_LOG_ONLINE_SESSION(Error, TEXT("EOS_Sessions_DestroySession() failed with EOS result code (%s)"), ANSI_TO_TCHAR(EOS_EResult_ToString(Data->ResultCode)));
+			UE_LOG_ONLINE_SESSION(Verbose, TEXT("Session [%s] not found"), *SessionName.ToString());
 		}
+
 		RemoveNamedSession(SessionName);
 		TriggerOnDestroySessionCompleteDelegates(SessionName, bWasSuccessful);
 	};
@@ -3773,12 +3818,19 @@ uint32 FOnlineSessionEOS::EndLobbySession(FNamedOnlineSession* Session)
 	// Only called from EndSession/DestroySession and presumes only in InProgress state
 	check(Session && Session->SessionState == EOnlineSessionState::InProgress);
 
-	EOSSubsystem->ExecuteNextTick([this, Session]()
+	EOSSubsystem->ExecuteNextTick([this, SessionName = Session->SessionName]()
+	{
+		if (FNamedOnlineSession* Session = GetNamedSession(SessionName))
 		{
-			Session->SessionState = EOnlineSessionState::Ended;
+			Session->SessionState = EOnlineSessionState::Ended;			
+		}
+		else
+		{
+			UE_LOG_ONLINE_SESSION(Verbose, TEXT("Session [%s] not found"), *SessionName.ToString());
+		}
 
-			TriggerOnEndSessionCompleteDelegates(Session->SessionName, true);
-		});
+		TriggerOnEndSessionCompleteDelegates(SessionName, true);
+	});
 
 	return ONLINE_IO_PENDING;
 }
