Added missing early error exit to FUserManagerEOS::QueryExternalIdMappings

diff --git a/Engine/Plugins/Online/OnlineSubsystemEOS/Source/OnlineSubsystemEOS/Private/UserManagerEOS.cpp b/Engine/Plugins/Online/OnlineSubsystemEOS/Source/OnlineSubsystemEOS/Private/UserManagerEOS.cpp
--- a/Engine/Plugins/Online/OnlineSubsystemEOS/Source/OnlineSubsystemEOS/Private/UserManagerEOS.cpp
+++ b/Engine/Plugins/Online/OnlineSubsystemEOS/Source/OnlineSubsystemEOS/Private/UserManagerEOS.cpp
@@ -2817,6 +2817,13 @@ bool FUserManagerEOS::QueryExternalIdMappings(const FUniqueNetId& UserId, const
 		Delegate.ExecuteIfBound(false, UserId, QueryOptions, ExternalIds, FString::Printf(TEXT("User (%s) is not logged in, so can't query external account ids"), *EOSID.UniqueNetIdStr));
 		return false;
 	}
+
+	if (ExternalIds.IsEmpty())
+	{
+		Delegate.ExecuteIfBound(false, UserId, QueryOptions, ExternalIds, FString::Printf(TEXT("List of ids to query is empty for User (%s), so can't query external account ids"), *EOSID.UniqueNetIdStr));
+		return false;
+	}
+
 	int32 LocalUserNum = GetLocalUserNumFromUniqueNetId(UserId);
 
 	// Mark the queries as in progress
