Additon of missing sessions API calls to implementation.

diff --git a/Engine/Plugins/Online/OnlineSubsystemEOS/Source/OnlineSubsystemEOS/Private/OnlineSessionEOS.cpp b/Engine/Plugins/Online/OnlineSubsystemEOS/Source/OnlineSubsystemEOS/Private/OnlineSessionEOS.cpp
--- a/Engine/Plugins/Online/OnlineSubsystemEOS/Source/OnlineSubsystemEOS/Private/OnlineSessionEOS.cpp
+++ b/Engine/Plugins/Online/OnlineSubsystemEOS/Source/OnlineSubsystemEOS/Private/OnlineSessionEOS.cpp
@@ -963,12 +963,42 @@ void FOnlineSessionEOS::SetPermissionLevel(EOS_HSessionModification SessionModHa
 		Options.PermissionLevel = EOS_EOnlineSessionPermissionLevel::EOS_OSPF_InviteOnly;
 	}
 
-	UE_LOG_ONLINE_SESSION(Log, TEXT("EOS_SessionModification_SetPermissionLevel() set to (%d)"), (int32)Options.PermissionLevel);
+	UE_LOG_ONLINE_SESSION(Log, TEXT("EOS_SessionModification_SetPermissionLevel() set to (%d) for session (%s)"), (int32)Options.PermissionLevel, *Session->SessionName.ToString());
 
 	EOS_EResult ResultCode = EOS_SessionModification_SetPermissionLevel(SessionModHandle, &Options);
 	if (ResultCode != EOS_EResult::EOS_Success)
 	{
-		UE_LOG_ONLINE_SESSION(Error, TEXT("EOS_SessionModification_SetPermissionLevel() failed with EOS result code (%s)"), ANSI_TO_TCHAR(EOS_EResult_ToString(ResultCode)));
+		UE_LOG_ONLINE_SESSION(Error, TEXT("EOS_SessionModification_SetPermissionLevel() failed with EOS result code (%s)"), *LexToString(ResultCode));
+	}
+}
+
+void FOnlineSessionEOS::SetMaxPlayers(EOS_HSessionModification SessionModHandle, FNamedOnlineSession* Session)
+{
+	EOS_SessionModification_SetMaxPlayersOptions Options = { };
+	Options.ApiVersion = EOS_SESSIONMODIFICATION_SETMAXPLAYERS_API_LATEST;
+	Options.MaxPlayers = Session->SessionSettings.NumPrivateConnections + Session->SessionSettings.NumPublicConnections;
+
+	UE_LOG_ONLINE_SESSION(Log, TEXT("EOS_SessionModification_SetMaxPlayers() set to (%d) for session (%s)"), Options.MaxPlayers, *Session->SessionName.ToString());
+
+	const EOS_EResult ResultCode = EOS_SessionModification_SetMaxPlayers(SessionModHandle, &Options);
+	if (ResultCode != EOS_EResult::EOS_Success)
+	{
+		UE_LOG_ONLINE_SESSION(Error, TEXT("EOS_SessionModification_SetMaxPlayers() failed with EOS result code (%s)"), *LexToString(ResultCode));
+	}
+}
+
+void FOnlineSessionEOS::SetInvitesAllowed(EOS_HSessionModification SessionModHandle, FNamedOnlineSession* Session)
+{
+	EOS_SessionModification_SetInvitesAllowedOptions Options = { };
+	Options.ApiVersion = EOS_SESSIONMODIFICATION_SETINVITESALLOWED_API_LATEST;
+	Options.bInvitesAllowed = Session->SessionSettings.bAllowInvites ? EOS_TRUE : EOS_FALSE;
+
+	UE_LOG_ONLINE_SESSION(Log, TEXT("EOS_SessionModification_SetInvitesAllowed() set to (%s) for session (%s)"), *LexToString(Options.bInvitesAllowed), *Session->SessionName.ToString());
+
+	const EOS_EResult ResultCode = EOS_SessionModification_SetInvitesAllowed(SessionModHandle, &Options);
+	if (ResultCode != EOS_EResult::EOS_Success)
+	{
+		UE_LOG_ONLINE_SESSION(Error, TEXT("EOS_SessionModification_SetInvitesAllowed() failed with EOS result code (%s)"), *LexToString(ResultCode));
 	}
 }
 
@@ -978,12 +1008,12 @@ void FOnlineSessionEOS::SetJoinInProgress(EOS_HSessionModification SessionModHan
 	Options.ApiVersion = EOS_SESSIONMODIFICATION_SETJOININPROGRESSALLOWED_API_LATEST;
 	Options.bAllowJoinInProgress = Session->SessionSettings.bAllowJoinInProgress ? EOS_TRUE : EOS_FALSE;
 
-	UE_LOG_ONLINE_SESSION(Log, TEXT("EOS_SessionModification_SetJoinInProgressAllowed() set to (%d)"), Options.bAllowJoinInProgress);
+	UE_LOG_ONLINE_SESSION(Log, TEXT("EOS_SessionModification_SetJoinInProgressAllowed() set to (%s) for session (%s)"), *LexToString(Options.bAllowJoinInProgress), *Session->SessionName.ToString());
 
 	EOS_EResult ResultCode = EOS_SessionModification_SetJoinInProgressAllowed(SessionModHandle, &Options);
 	if (ResultCode != EOS_EResult::EOS_Success)
 	{
-		UE_LOG_ONLINE_SESSION(Error, TEXT("EOS_SessionModification_SetJoinInProgressAllowed() failed with EOS result code (%s)"), ANSI_TO_TCHAR(EOS_EResult_ToString(ResultCode)));
+		UE_LOG_ONLINE_SESSION(Error, TEXT("EOS_SessionModification_SetJoinInProgressAllowed() failed with EOS result code (%s)"), *LexToString(ResultCode));
 	}
 }
 
@@ -994,12 +1024,12 @@ void FOnlineSessionEOS::AddAttribute(EOS_HSessionModification SessionModHandle,
 	Options.AdvertisementType = EOS_ESessionAttributeAdvertisementType::EOS_SAAT_Advertise;
 	Options.SessionAttribute = Attribute;
 
-	UE_LOG_ONLINE_SESSION(Log, TEXT("EOS_SessionModification_AddAttribute() named (%s) with value (%s)"), *FString(Attribute->Key), *MakeStringFromAttributeValue(Attribute));
+	UE_LOG_ONLINE_SESSION(Log, TEXT("EOS_SessionModification_AddAttribute() named (%s) with value (%s)"), UTF8_TO_TCHAR(Attribute->Key), *MakeStringFromAttributeValue(Attribute));
 
 	EOS_EResult ResultCode = EOS_SessionModification_AddAttribute(SessionModHandle, &Options);
 	if (ResultCode != EOS_EResult::EOS_Success)
 	{
-		UE_LOG_ONLINE_SESSION(Error, TEXT("EOS_SessionModification_AddAttribute() failed for attribute name (%s) with EOS result code (%s)"), *FString(Attribute->Key), ANSI_TO_TCHAR(EOS_EResult_ToString(ResultCode)));
+		UE_LOG_ONLINE_SESSION(Error, TEXT("EOS_SessionModification_AddAttribute() failed for attribute name (%s) with EOS result code (%s)"), *FString(Attribute->Key), *LexToString(ResultCode));
 	}
 }
 
@@ -1230,8 +1260,12 @@ uint32 FOnlineSessionEOS::CreateEOSSession(int32 HostingPlayerNum, FNamedOnlineS
 
 uint32 FOnlineSessionEOS::SharedSessionUpdate(EOS_HSessionModification SessionModHandle, FNamedOnlineSession* Session, FUpdateSessionCallback* Callback)
 {
-	// Set invite flags
+	// Set joinability flags
 	SetPermissionLevel(SessionModHandle, Session);
+	// Set max players
+	SetMaxPlayers(SessionModHandle, Session);
+	// Set invite flags
+	SetInvitesAllowed(SessionModHandle, Session);
 	// Set JIP flag
 	SetJoinInProgress(SessionModHandle, Session);
 	// Add any attributes for filtering by searchers
@@ -1901,33 +1935,26 @@ void FOnlineSessionEOS::AddLobbySearchAttribute(EOS_HLobbySearch LobbySearchHand
 
 void FOnlineSessionEOS::CopySearchResult(EOS_HSessionDetails SessionHandle, EOS_SessionDetails_Info* SessionInfo, FOnlineSession& OutSession)
 {
+	OutSession.NumOpenPrivateConnections = SessionInfo->NumOpenPublicConnections;
+	OutSession.SessionSettings.NumPrivateConnections = SessionInfo->Settings->NumPublicConnections;
 	OutSession.SessionSettings.bAllowJoinInProgress = SessionInfo->Settings->bAllowJoinInProgress == EOS_TRUE;
+	OutSession.SessionSettings.bAllowInvites = SessionInfo->Settings->bInvitesAllowed == EOS_TRUE;
 	switch (SessionInfo->Settings->PermissionLevel)
 	{
 		case EOS_EOnlineSessionPermissionLevel::EOS_OSPF_InviteOnly:
 		{
-			OutSession.SessionSettings.bAllowInvites = true;
-			OutSession.SessionSettings.bUsesPresence = false;
+			OutSession.SessionSettings.bUsesPresence = true;
 			OutSession.SessionSettings.bAllowJoinViaPresence = false;
 			break;
 		}
 		case EOS_EOnlineSessionPermissionLevel::EOS_OSPF_JoinViaPresence:
-		{
-			OutSession.SessionSettings.bAllowInvites = true;
-			OutSession.SessionSettings.bUsesPresence = true;
-			OutSession.SessionSettings.bAllowJoinViaPresence = true;
-			break;
-		}
 		case EOS_EOnlineSessionPermissionLevel::EOS_OSPF_PublicAdvertised:
 		{
-			OutSession.SessionSettings.bAllowInvites = true;
 			OutSession.SessionSettings.bUsesPresence = true;
 			OutSession.SessionSettings.bAllowJoinViaPresence = true;
 			break;
 		}
 	}
-	OutSession.SessionSettings.NumPrivateConnections = SessionInfo->Settings->NumPublicConnections;
-	OutSession.NumOpenPrivateConnections = SessionInfo->NumOpenPublicConnections;
 
 	CopyAttributes(SessionHandle, OutSession);
 }
diff --git a/Engine/Plugins/Online/OnlineSubsystemEOS/Source/OnlineSubsystemEOS/Private/OnlineSessionEOS.h b/Engine/Plugins/Online/OnlineSubsystemEOS/Source/OnlineSubsystemEOS/Private/OnlineSessionEOS.h
--- a/Engine/Plugins/Online/OnlineSubsystemEOS/Source/OnlineSubsystemEOS/Private/OnlineSessionEOS.h
+++ b/Engine/Plugins/Online/OnlineSubsystemEOS/Source/OnlineSubsystemEOS/Private/OnlineSessionEOS.h
@@ -295,6 +295,8 @@ private:
 	void CopyAttributes(EOS_HSessionDetails SessionHandle, FOnlineSession& OutSession);
 
 	void SetPermissionLevel(EOS_HSessionModification SessionModHandle, FNamedOnlineSession* Session);
+	void SetMaxPlayers(EOS_HSessionModification SessionModHandle, FNamedOnlineSession* Session);
+	void SetInvitesAllowed(EOS_HSessionModification SessionModHandle, FNamedOnlineSession* Session);
 	void SetJoinInProgress(EOS_HSessionModification SessionModHandle, FNamedOnlineSession* Session);
 	void AddAttribute(EOS_HSessionModification SessionModHandle, const EOS_Sessions_AttributeData* Attribute);
 	void SetAttributes(EOS_HSessionModification SessionModHandle, FNamedOnlineSession* Session);
