Changed BucketId in EOS Sessions to use BuildUniqueId to fix  compatibility issues.

diff --git a/Engine/Plugins/Online/OnlineSubsystemEOS/Source/OnlineSubsystemEOS/Private/OnlineSubsystemEOS.cpp b/Engine/Plugins/Online/OnlineSubsystemEOS/Source/OnlineSubsystemEOS/Private/OnlineSubsystemEOS.cpp
--- a/Engine/Plugins/Online/OnlineSubsystemEOS/Source/OnlineSubsystemEOS/Private/OnlineSubsystemEOS.cpp
+++ b/Engine/Plugins/Online/OnlineSubsystemEOS/Source/OnlineSubsystemEOS/Private/OnlineSubsystemEOS.cpp
@@ -404,7 +404,7 @@ bool FOnlineSubsystemEOS::Init()
 	UserManager->Init();
 	SessionInterfacePtr = MakeShareable(new FOnlineSessionEOS(this));
 	// Set the bucket id to use for all sessions based upon the name and version to avoid upgrade issues
-	SessionInterfacePtr->Init(EOSSDKManager->GetProductName() + TEXT("_") + EOSSDKManager->GetProductVersion());
+	SessionInterfacePtr->Init(EOSSDKManager->GetProductName() + TEXT("_") + FString::FromInt(GetBuildUniqueId()));
 	StatsInterfacePtr = MakeShareable(new FOnlineStatsEOS(this));
 	LeaderboardsInterfacePtr = MakeShareable(new FOnlineLeaderboardsEOS(this));
 	AchievementsInterfacePtr = MakeShareable(new FOnlineAchievementsEOS(this));
