Added logic to cache the updated Auth Token after a refresh in UserManagerEOS

diff --git a/Engine/Plugins/Online/OnlineSubsystemEOS/Source/OnlineSubsystemEOS/Private/UserManagerEOS.cpp b/Engine/Plugins/Online/OnlineSubsystemEOS/Source/OnlineSubsystemEOS/Private/UserManagerEOS.cpp
--- a/Engine/Plugins/Online/OnlineSubsystemEOS/Source/OnlineSubsystemEOS/Private/UserManagerEOS.cpp
+++ b/Engine/Plugins/Online/OnlineSubsystemEOS/Source/OnlineSubsystemEOS/Private/UserManagerEOS.cpp
@@ -768,6 +768,12 @@ void FUserManagerEOS::RefreshConnectLogin(int32 LocalUserNum)
 		EOS_EResult CopyResult = EOS_Auth_CopyUserAuthToken(EOSSubsystem->AuthHandle, &CopyOptions, AccountId, &AuthToken);
 		if (CopyResult == EOS_EResult::EOS_Success)
 		{
+			// We update the auth token cached in the user account, along with the user information
+			const FUniqueNetIdEOSPtr UniqueNetId = UserNumToNetIdMap.FindChecked(LocalUserNum);
+			const FUserOnlineAccountEOSRef UserAccountRef = StringToUserAccountMap.FindChecked(UniqueNetId->ToString());
+			UserAccountRef->SetAuthAttribute(AUTH_ATTR_ID_TOKEN, AuthToken->AccessToken);
+			UpdateUserInfo(UserAccountRef, AccountId, AccountId);
+
 			EOS_Connect_Credentials Credentials = { };
 			Credentials.ApiVersion = EOS_CONNECT_CREDENTIALS_API_LATEST;
 			Credentials.Type = EOS_EExternalCredentialType::EOS_ECT_EPIC;
