Added missing check on delegate trigger

diff --git a/Engine/Plugins/Online/OnlineSubsystemEOS/Source/OnlineSubsystemEOS/Private/OnlineSessionEOS.cpp b/Engine/Plugins/Online/OnlineSubsystemEOS/Source/OnlineSubsystemEOS/Private/OnlineSessionEOS.cpp
--- a/Engine/Plugins/Online/OnlineSubsystemEOS/Source/OnlineSubsystemEOS/Private/OnlineSessionEOS.cpp
+++ b/Engine/Plugins/Online/OnlineSubsystemEOS/Source/OnlineSubsystemEOS/Private/OnlineSessionEOS.cpp
@@ -2212,7 +2212,8 @@ void FOnlineSessionEOS::FindEOSSessionById(int32 LocalUserNum, const FUniqueNetI
 			LocalSessionSearch->SearchState = EOnlineAsyncTaskState::Failed;
 			UE_LOG_ONLINE_SESSION(Error, TEXT("EOS_SessionSearch_Find() failed with EOS result code (%s)"), ANSI_TO_TCHAR(EOS_EResult_ToString(Data->ResultCode)));
 		}
-		OnComplete.ExecuteIfBound(LocalUserNum, LocalSessionSearch->SearchState == EOnlineAsyncTaskState::Done, LocalSessionSearch->SearchResults.Last());
+
+		OnComplete.ExecuteIfBound(LocalUserNum, LocalSessionSearch->SearchState == EOnlineAsyncTaskState::Done, !LocalSessionSearch->SearchResults.IsEmpty() ? LocalSessionSearch->SearchResults.Last() : FOnlineSessionSearchResult());
 	};
 
 	EOS_SessionSearch_FindOptions FindOptions = { };
