- Fix RefreshConnectLogin and GetLinkedAccountAuthToken assuming a token will have the data field set.

diff --git a/Engine/Plugins/Online/OnlineSubsystemEOS/Source/OnlineSubsystemEOS/Private/UserManagerEOS.cpp b/Engine/Plugins/Online/OnlineSubsystemEOS/Source/OnlineSubsystemEOS/Private/UserManagerEOS.cpp
--- a/Engine/Plugins/Online/OnlineSubsystemEOS/Source/OnlineSubsystemEOS/Private/UserManagerEOS.cpp
+++ b/Engine/Plugins/Online/OnlineSubsystemEOS/Source/OnlineSubsystemEOS/Private/UserManagerEOS.cpp
@@ -808,7 +808,7 @@ void FUserManagerEOS::RefreshConnectLogin(int32 LocalUserNum)
 		GetPlatformAuthToken(LocalUserNum,
 			FOnGetLinkedAccountAuthTokenCompleteDelegate::CreateLambda([this](int32 LocalUserNum, bool bWasSuccessful, const FExternalAuthToken& AuthToken)
 			{
-				if (!bWasSuccessful || !AuthToken.HasTokenData())
+				if (!bWasSuccessful || !AuthToken.IsValid())
 				{
 					UE_LOG_ONLINE(Error, TEXT("ConnectLoginNoEAS(%d) failed due to the platform OSS giving an empty auth token"), LocalUserNum);
 					Logout(LocalUserNum);
@@ -1536,7 +1536,7 @@ void FUserManagerEOS::GetLinkedAccountAuthToken(int32 LocalUserNum, const FOnGet
 {
 	FExternalAuthToken ExternalToken;
 	ExternalToken.TokenString = GetAuthToken(LocalUserNum);
-	Delegate.ExecuteIfBound(LocalUserNum, ExternalToken.HasTokenString(), ExternalToken);
+	Delegate.ExecuteIfBound(LocalUserNum, ExternalToken.IsValid(), ExternalToken);
 }
 
 void FUserManagerEOS::GetUserPrivilege(const FUniqueNetId& UserId, EUserPrivileges::Type Privilege, const FOnGetUserPrivilegeCompleteDelegate& Delegate)
