- Fix OnlineStoreEOS to use the DecimalPoint value from EOS, not the one from UE I18N, as they disagree on the existence of chilean centavos.

diff --git a/Engine/Plugins/Online/OnlineSubsystemEOS/Source/OnlineSubsystemEOS/Private/OnlineStoreEOS.cpp b/Engine/Plugins/Online/OnlineSubsystemEOS/Source/OnlineSubsystemEOS/Private/OnlineStoreEOS.cpp
--- a/Engine/Plugins/Online/OnlineSubsystemEOS/Source/OnlineSubsystemEOS/Private/OnlineStoreEOS.cpp
+++ b/Engine/Plugins/Online/OnlineSubsystemEOS/Source/OnlineSubsystemEOS/Private/OnlineStoreEOS.cpp
@@ -112,6 +112,9 @@ void FOnlineStoreEOS::QueryOffers(const FUniqueNetId& UserId, const FOnQueryOnli
 				OfferRef->DiscountType = Offer->DiscountPercentage == 0 ? EOnlineStoreOfferDiscountType::NotOnSale : EOnlineStoreOfferDiscountType::DiscountAmount;
 			}
 
+			OfferRef->RegularPriceText = FText::AsCurrencyBase(OfferRef->RegularPrice, OfferRef->CurrencyCode, NULL, Offer->DecimalPoint);
+			OfferRef->PriceText = FText::AsCurrencyBase(OfferRef->NumericPrice, OfferRef->CurrencyCode, NULL, Offer->DecimalPoint);
+
 			CachedOffers.Add(OfferRef);
 			CachedOfferIds.Add(OfferRef->OfferId);
 
