- Refactor FUniqueNetIdEOS to not depend on FUniqueNetIdString
- Add EOS net id registry to guarantee a single net id object for a given user.
- Update EOS net id hash to use object pointer.
- Post-creation mutation of eos net ids (e.g. addition of missing PUID) keeps the hash stable, therefore not breaking TUniqueNetIdMap/FUniqueNetIdSet.
- Fix EOS and EOSPlus being incompatible and crashing when attempting to play together.
- Adjust CreateUniquePlayerIdWrapper to prefer to load a net id with the mapped OSS if one is found.

diff --git a/Engine/Plugins/Online/OnlineSubsystemEOS/Source/OnlineSubsystemEOS/Private/NboSerializerEOS.h b/Engine/Plugins/Online/OnlineSubsystemEOS/Source/OnlineSubsystemEOS/Private/NboSerializerEOS.h
--- a/Engine/Plugins/Online/OnlineSubsystemEOS/Source/OnlineSubsystemEOS/Private/NboSerializerEOS.h
+++ b/Engine/Plugins/Online/OnlineSubsystemEOS/Source/OnlineSubsystemEOS/Private/NboSerializerEOS.h
@@ -38,15 +38,6 @@ public:
 		return Ar;
  	}
 
-	/**
-	 * Adds EOS Unique Id to the buffer
-	 */
-	friend inline FNboSerializeToBufferEOS& operator<<(FNboSerializeToBufferEOS& Ar, const FUniqueNetIdEOS& UniqueId)
-	{
-		((FNboSerializeToBuffer&)Ar) << UniqueId.UniqueNetIdStr;
-		return Ar;
-	}
-
 	/**
 	 * Adds string Unique Id to the buffer
 	 */
@@ -77,22 +68,16 @@ public:
  	friend inline FNboSerializeFromBufferEOS& operator>>(FNboSerializeFromBufferEOS& Ar, FOnlineSessionInfoEOS& SessionInfo)
  	{
 		check(SessionInfo.HostAddr.IsValid());
+
+		FString SessionId;
+		Ar >> SessionId;
 		// Skip SessionType (assigned at creation)
-		SessionInfo.SessionId = FUniqueNetIdEOS::Create();
-		Ar >> *ConstCastSharedRef<FUniqueNetIdString>(SessionInfo.SessionId);
 		Ar >> *SessionInfo.HostAddr;
-		return Ar;
- 	}
 
-	/**
-	 * Reads EOS Unique Id from the buffer
-	 */
-	friend inline FNboSerializeFromBufferEOS& operator>>(FNboSerializeFromBufferEOS& Ar, FUniqueNetIdEOS& UniqueId)
-	{
-		Ar >> UniqueId.UniqueNetIdStr;
-		UniqueId.ParseAccountIds();
+		SessionInfo.SessionId = FUniqueNetIdString::Create(MoveTemp(SessionId), FName("EOS"));
+
 		return Ar;
-	}
+ 	}
 
 	/**
 	 * Reads string Unique Id from the buffer
diff --git a/Engine/Plugins/Online/OnlineSubsystemEOS/Source/OnlineSubsystemEOS/Private/OnlineLeaderboardsEOS.cpp b/Engine/Plugins/Online/OnlineSubsystemEOS/Source/OnlineSubsystemEOS/Private/OnlineLeaderboardsEOS.cpp
--- a/Engine/Plugins/Online/OnlineSubsystemEOS/Source/OnlineSubsystemEOS/Private/OnlineLeaderboardsEOS.cpp
+++ b/Engine/Plugins/Online/OnlineSubsystemEOS/Source/OnlineSubsystemEOS/Private/OnlineLeaderboardsEOS.cpp
@@ -309,7 +309,7 @@ bool FOnlineLeaderboardsEOS::ReadLeaderboardsAroundRank(int32 Rank, uint32 Range
 						Nickname = TEXT("Unknown Player");
 					}
 					const EOS_EpicAccountId AccountId = EOS_EpicAccountId_FromString(EpicIdStr);
-					const FUniqueNetIdEOSRef NetId = FUniqueNetIdEOS::Create(AccountId, Record->UserId);
+					const FUniqueNetIdEOSRef NetId = FUniqueNetIdEOSRegistry::FindOrAdd(AccountId, Record->UserId).ToSharedRef();
 
 					FOnlineStatsRow* Row = new(LambdaReadObject->Rows) FOnlineStatsRow(Nickname, NetId);
 					Row->Rank = Record->Rank;
diff --git a/Engine/Plugins/Online/OnlineSubsystemEOS/Source/OnlineSubsystemEOS/Private/OnlineSessionEOS.cpp b/Engine/Plugins/Online/OnlineSubsystemEOS/Source/OnlineSubsystemEOS/Private/OnlineSessionEOS.cpp
--- a/Engine/Plugins/Online/OnlineSubsystemEOS/Source/OnlineSubsystemEOS/Private/OnlineSessionEOS.cpp
+++ b/Engine/Plugins/Online/OnlineSubsystemEOS/Source/OnlineSubsystemEOS/Private/OnlineSessionEOS.cpp
@@ -1988,7 +1988,7 @@ void FOnlineSessionEOS::CopyAttributes(EOS_HSessionDetails SessionHandle, FOnlin
 			}
 			else if (Key == TEXT("OwningUserId"))
 			{
-				OutSession.OwningUserId = FUniqueNetIdEOS::Create(UTF8_TO_TCHAR(Attribute->Data->Value.AsUtf8));
+				OutSession.OwningUserId = FUniqueNetIdEOSRegistry::FindOrAdd(UTF8_TO_TCHAR(Attribute->Data->Value.AsUtf8));
 			}
 			else if (Key == TEXT("OwningUserName"))
 			{
@@ -3174,13 +3174,13 @@ void FOnlineSessionEOS::ReadSessionFromPacket(FNboSerializeFromBufferEOS& Packet
 #endif
 
 	/** Owner of the session */
-	FUniqueNetIdEOSRef OwningUserId = FUniqueNetIdEOS::Create();
-	Packet >> *ConstCastSharedRef<FUniqueNetIdEOS>(OwningUserId)
+	FString OwningUserIdStr;
+	Packet >> OwningUserIdStr
 		>> Session->OwningUserName
 		>> Session->NumOpenPrivateConnections
 		>> Session->NumOpenPublicConnections;
 
-	Session->OwningUserId = OwningUserId;
+	Session->OwningUserId = FUniqueNetIdEOSRegistry::FindOrAdd(OwningUserIdStr);
 
 	// Allocate and read the connection data
 	FOnlineSessionInfoEOS* EOSSessionInfo = new FOnlineSessionInfoEOS();
@@ -4194,7 +4194,7 @@ void FOnlineSessionEOS::CopyLobbyAttributes(const TSharedRef<FLobbyDetailsEOS>&
 			FString Key = UTF8_TO_TCHAR(Attribute->Data->Key);
 			if (Key == TEXT("OwningUserId"))
 			{
-				OutSession.OwningUserId = FUniqueNetIdEOS::Create(UTF8_TO_TCHAR(Attribute->Data->Value.AsUtf8));
+				OutSession.OwningUserId = FUniqueNetIdEOSRegistry::FindOrAdd(UTF8_TO_TCHAR(Attribute->Data->Value.AsUtf8));
 			}
 			else if (Key == TEXT("OwningUserName"))
 			{
diff --git a/Engine/Plugins/Online/OnlineSubsystemEOS/Source/OnlineSubsystemEOS/Private/OnlineStatsEOS.cpp b/Engine/Plugins/Online/OnlineSubsystemEOS/Source/OnlineSubsystemEOS/Private/OnlineStatsEOS.cpp
--- a/Engine/Plugins/Online/OnlineSubsystemEOS/Source/OnlineSubsystemEOS/Private/OnlineStatsEOS.cpp
+++ b/Engine/Plugins/Online/OnlineSubsystemEOS/Source/OnlineSubsystemEOS/Private/OnlineStatsEOS.cpp
@@ -328,7 +328,7 @@ void FOnlineStatsEOS::UpdateStats(const FUniqueNetIdRef LocalUserId, const TArra
 	const EOS_ProductUserId UserId = EOSId.GetProductUserId();
 	if (UserId == nullptr)
 	{
-		UE_LOG_ONLINE_STATS(Error, TEXT("UpdateStats() failed for unknown player (%s)"), *EOSId.UniqueNetIdStr);
+		UE_LOG_ONLINE_STATS(Error, TEXT("UpdateStats() failed for unknown player (%s)"), *EOSId.ToDebugString());
 		Delegate.ExecuteIfBound(FOnlineError(EOnlineErrorResult::InvalidCreds));
 		return;
 	}
diff --git a/Engine/Plugins/Online/OnlineSubsystemEOS/Source/OnlineSubsystemEOS/Private/OnlineSubsystemEOSTypes.cpp b/Engine/Plugins/Online/OnlineSubsystemEOS/Source/OnlineSubsystemEOS/Private/OnlineSubsystemEOSTypes.cpp
new file mode 100644
--- /dev/null
+++ b/Engine/Plugins/Online/OnlineSubsystemEOS/Source/OnlineSubsystemEOS/Private/OnlineSubsystemEOSTypes.cpp
@@ -0,0 +1,226 @@
+// Copyright Epic Games, Inc. All Rights Reserved.
+
+#include "OnlineSubsystemEOSTypes.h"
+
+#include "Algo/AnyOf.h"
+#include "Misc/LazySingleton.h"
+#include "Misc/ScopeRWLock.h"
+
+const FUniqueNetIdEOS& FUniqueNetIdEOS::Cast(const FUniqueNetId& NetId)
+{
+	check(GetTypeStatic() == NetId.GetType());
+	return *static_cast<const FUniqueNetIdEOS*>(&NetId);
+}
+
+const FUniqueNetIdEOSRef& FUniqueNetIdEOS::EmptyId()
+{
+	static const FUniqueNetIdEOSRef EmptyId(Create());
+	return EmptyId;
+}
+
+FName FUniqueNetIdEOS::GetTypeStatic()
+{
+	static FName NAME_Eos(TEXT("EOS"));
+	return NAME_Eos;
+}
+
+FName FUniqueNetIdEOS::GetType() const
+{
+	return GetTypeStatic();
+}
+
+const uint8* FUniqueNetIdEOS::GetBytes() const
+{
+	return RawBytes;
+}
+
+int32 FUniqueNetIdEOS::GetSize() const
+{
+	return EOS_ID_BYTE_SIZE;
+}
+
+bool FUniqueNetIdEOS::IsValid() const
+{
+	return EOS_EpicAccountId_IsValid(EpicAccountId) || EOS_ProductUserId_IsValid(ProductUserId);
+}
+
+uint32 FUniqueNetIdEOS::GetTypeHash() const
+{
+	return ::GetTypeHash(static_cast<const void*>(this));
+}
+
+FString FUniqueNetIdEOS::ToString() const
+{
+	return LexToString(EpicAccountId) + EOS_ID_SEPARATOR + LexToString(ProductUserId);
+}
+
+FString FUniqueNetIdEOS::ToDebugString() const
+{
+	if (IsValid())
+	{
+		const FString EpicAccountIdStr = LexToString(EpicAccountId);
+		const FString ProductUserIdStr = LexToString(ProductUserId);
+		return OSS_UNIQUEID_REDACT(*this, EpicAccountIdStr) + EOS_ID_SEPARATOR + OSS_UNIQUEID_REDACT(*this, ProductUserIdStr);
+	}
+	else
+	{
+		return TEXT("INVALID");
+	}
+}
+
+FUniqueNetIdEOS::FUniqueNetIdEOS(const uint8* Bytes, int32 Size)
+{
+	check(Size == EOS_ID_BYTE_SIZE);
+	FMemory::Memcpy(RawBytes, Bytes, EOS_ID_BYTE_SIZE);
+
+	const bool bIsEasNonZero = Algo::AnyOf(TArrayView<const uint8>(Bytes, ID_HALF_BYTE_SIZE));
+	if (bIsEasNonZero)
+	{
+		const FString EpicAccountIdStr = BytesToHex(Bytes, ID_HALF_BYTE_SIZE);
+		EpicAccountId = EOS_EpicAccountId_FromString(TCHAR_TO_UTF8(*EpicAccountIdStr));
+	}
+
+	const bool bIsPuidNonZero = Algo::AnyOf(TArrayView<const uint8>(Bytes + ID_HALF_BYTE_SIZE, ID_HALF_BYTE_SIZE));
+	if (bIsPuidNonZero)
+	{
+		const FString ProductUserIdStr = BytesToHex(Bytes + ID_HALF_BYTE_SIZE, ID_HALF_BYTE_SIZE);
+		ProductUserId = EOS_ProductUserId_FromString(TCHAR_TO_UTF8(*ProductUserIdStr));
+	}
+}
+
+FUniqueNetIdEOS::FUniqueNetIdEOS(EOS_EpicAccountId InEpicAccountId, EOS_ProductUserId InProductUserId)
+	: EpicAccountId(InEpicAccountId)
+	, ProductUserId(InProductUserId)
+{
+	HexToBytes(LexToString(EpicAccountId), RawBytes);
+	HexToBytes(LexToString(ProductUserId), RawBytes + ID_HALF_BYTE_SIZE);
+}
+
+FUniqueNetIdEOSRegistry& FUniqueNetIdEOSRegistry::Get()
+{
+	return TLazySingleton<FUniqueNetIdEOSRegistry>::Get();
+}
+
+FUniqueNetIdEOSPtr FUniqueNetIdEOSRegistry::FindOrAddImpl(const FString& NetIdStr)
+{
+	FString EpicAccountIdStr;
+	FString ProductUserIdStr;
+	if (!NetIdStr.Split(EOS_ID_SEPARATOR, &EpicAccountIdStr, &ProductUserIdStr)
+		|| (!EpicAccountIdStr.IsEmpty() && EpicAccountIdStr.Len() != EOS_EPICACCOUNTID_MAX_LENGTH)
+		|| (!ProductUserIdStr.IsEmpty() && ProductUserIdStr.Len() != EOS_PRODUCTUSERID_MAX_LENGTH))
+	{
+		return nullptr;
+	}
+
+	EOS_EpicAccountId EpicAccountId = nullptr;
+	if (!EpicAccountIdStr.IsEmpty())
+	{
+		EpicAccountId = EOS_EpicAccountId_FromString(TCHAR_TO_UTF8(*EpicAccountIdStr));
+	}
+	EOS_ProductUserId ProductUserId = nullptr;
+	if (!ProductUserIdStr.IsEmpty())
+	{
+		ProductUserId = EOS_ProductUserId_FromString(TCHAR_TO_UTF8(*ProductUserIdStr));
+	}
+
+	return FindOrAddImpl(EpicAccountId, ProductUserId);
+}
+
+FUniqueNetIdEOSPtr FUniqueNetIdEOSRegistry::FindOrAddImpl(const uint8* Bytes, int32 Size)
+{
+	if (Size == EOS_ID_BYTE_SIZE)
+	{
+		// The net id type already knows how to deserialize, so create a temp one to get the EAS/PUID
+		FUniqueNetIdEOS Temp(Bytes, Size);
+		return FindOrAddImpl(Temp.GetEpicAccountId(), Temp.GetProductUserId());
+	}
+	return nullptr;
+}
+
+FUniqueNetIdEOSPtr FUniqueNetIdEOSRegistry::FindOrAddImpl(const EOS_EpicAccountId InEpicAccountId, const EOS_ProductUserId InProductUserId)
+{
+	FUniqueNetIdEOSPtr Result;
+	const bool bInEpicAccountIdValid = EOS_EpicAccountId_IsValid(InEpicAccountId) == EOS_TRUE;
+	const bool bInProductUserIdValid = EOS_ProductUserId_IsValid(InProductUserId) == EOS_TRUE;
+	if (bInEpicAccountIdValid || bInProductUserIdValid)
+	{
+		bool bUpdateEpicAccountId = false;
+		bool bUpdateProductUserId = false;
+
+		auto FindExisting = [this, InEpicAccountId, InProductUserId, bInEpicAccountIdValid, bInProductUserIdValid, &bUpdateEpicAccountId, &bUpdateProductUserId]()
+		{
+			FUniqueNetIdEOSPtr Result;
+			if (const FUniqueNetIdEOSRef* FoundEas = bInEpicAccountIdValid ? EasToNetId.Find(InEpicAccountId) : nullptr)
+			{
+				Result = *FoundEas;
+			}
+			else if (const FUniqueNetIdEOSRef* FoundProd = bInProductUserIdValid ? PuidToNetId.Find(InProductUserId) : nullptr)
+			{
+				Result = *FoundProd;
+			}
+
+			if (Result.IsValid())
+			{
+				const EOS_EpicAccountId FoundEpicAccountId = Result->GetEpicAccountId();
+				const EOS_ProductUserId FoundProductUserId = Result->GetProductUserId();
+				const bool bFoundEpicAccountIdValid = EOS_EpicAccountId_IsValid(FoundEpicAccountId) == EOS_TRUE;
+				const bool bFoundProductUserIdValid = EOS_ProductUserId_IsValid(FoundProductUserId) == EOS_TRUE;
+
+				// Check that the found EAS/EOS ids are either unset, or match the input. If a valid input is passed for a currently unset field, this is an update
+				check(!bFoundEpicAccountIdValid || !bInEpicAccountIdValid || InEpicAccountId == FoundEpicAccountId);
+				check(!bFoundProductUserIdValid || !bInProductUserIdValid || InProductUserId == FoundProductUserId);
+				bUpdateEpicAccountId = !bFoundEpicAccountIdValid && bInEpicAccountIdValid;
+				bUpdateProductUserId = !bFoundProductUserIdValid && bInProductUserIdValid;
+
+			}
+			return Result;
+		};
+
+
+		{
+			// First take read lock and look for existing elements
+			const FReadScopeLock ReadLock(Lock);
+			Result = FindExisting();
+		}
+
+		if (!Result.IsValid())
+		{
+			// Double-checked locking. If we didn't find an element, we take the write lock, and look again, in case another thread raced with us and added one.
+			const FWriteScopeLock WriteLock(Lock);
+			Result = FindExisting();
+
+			if (!Result.IsValid())
+			{
+				// if we didn't find one we can create a new one
+				Result = FUniqueNetIdEOS::Create(InEpicAccountId, InProductUserId);
+				if (bInEpicAccountIdValid)
+				{
+					EasToNetId.Emplace(InEpicAccountId, Result.ToSharedRef());
+				}
+				if (bInProductUserIdValid)
+				{
+					PuidToNetId.Emplace(InProductUserId, Result.ToSharedRef());
+				}
+			}
+		}
+
+		check(Result.IsValid());
+		if (bUpdateEpicAccountId || bUpdateProductUserId)
+		{
+			// Finally, update any previously unset fields for which we now have a valid value.
+			const FWriteScopeLock WriteLock(Lock);
+			if (bUpdateEpicAccountId)
+			{
+				EasToNetId.Emplace(InEpicAccountId, Result.ToSharedRef());
+				*ConstCastSharedPtr<FUniqueNetIdEOS>(Result) = FUniqueNetIdEOS(InEpicAccountId, Result->GetProductUserId());
+			}
+			if (bUpdateProductUserId)
+			{
+				PuidToNetId.Emplace(InProductUserId, Result.ToSharedRef());
+				*ConstCastSharedPtr<FUniqueNetIdEOS>(Result) = FUniqueNetIdEOS(Result->GetEpicAccountId(), InProductUserId);
+			}
+		}
+	}
+
+	return Result;
+}
diff --git a/Engine/Plugins/Online/OnlineSubsystemEOS/Source/OnlineSubsystemEOS/Private/OnlineSubsystemEOSTypes.h b/Engine/Plugins/Online/OnlineSubsystemEOS/Source/OnlineSubsystemEOS/Private/OnlineSubsystemEOSTypes.h
--- a/Engine/Plugins/Online/OnlineSubsystemEOS/Source/OnlineSubsystemEOS/Private/OnlineSubsystemEOSTypes.h
+++ b/Engine/Plugins/Online/OnlineSubsystemEOS/Source/OnlineSubsystemEOS/Private/OnlineSubsystemEOSTypes.h
@@ -2,82 +2,47 @@
 
 #pragma once
 
-#include "OnlineSubsystemTypes.h"
-
-#include "OnlineSubsystem.h"
 #include "Interfaces/OnlineIdentityInterface.h"
-#include "Interfaces/OnlineFriendsInterface.h"
 #include "Interfaces/OnlinePresenceInterface.h"
 #include "Interfaces/OnlineUserInterface.h"
-#include "EOSSharedTypes.h"
 #include "EOSShared.h"
+#include "EOSSharedTypes.h"
 #include "IPAddress.h"
+#include "OnlineSubsystem.h"
+#include "OnlineSubsystemTypes.h"
 
 #define EOS_OSS_STRING_BUFFER_LENGTH 256 + 1 // 256 plus null terminator
 
 class FOnlineSubsystemEOS;
 
+#define EOSPLUS_ID_SEPARATOR TEXT("_+_")
 #define EOS_ID_SEPARATOR TEXT("|")
-#define EMPTY_EASID TEXT("00000000000000000000000000000000")
-#define EMPTY_PUID TEXT("00000000000000000000000000000000")
 #define ID_HALF_BYTE_SIZE 16
 #define EOS_ID_BYTE_SIZE (ID_HALF_BYTE_SIZE * 2)
 
 typedef TSharedPtr<const class FUniqueNetIdEOS> FUniqueNetIdEOSPtr;
 typedef TSharedRef<const class FUniqueNetIdEOS> FUniqueNetIdEOSRef;
 
-static inline FString MakeNetIdStringFromIds(EOS_EpicAccountId AccountId, EOS_ProductUserId UserId)
-{
-	return LexToString(AccountId) + EOS_ID_SEPARATOR + LexToString(UserId);
-}
-
 /**
- * Unique net id wrapper for a EOS account ids. The underlying string is a combination
- * of both account ids concatenated. "<EOS_EpicAccountId>|<EOS_ProductAccountId>"
+ * Unique net id wrapper for a EOS account ids.
  */
-class FUniqueNetIdEOS :
-	public FUniqueNetIdString
+class FUniqueNetIdEOS : public FUniqueNetId
 {
 public:
-	template<typename... TArgs>
-	static FUniqueNetIdEOSRef Create(TArgs&&... Args)
-	{
-		return MakeShareable(new FUniqueNetIdEOS(Forward<TArgs>(Args)...));
-	}
-
-	static const FUniqueNetIdEOS& Cast(const FUniqueNetId& NetId)
-	{
-		check(GetTypeStatic() == NetId.GetType());
-		return *static_cast<const FUniqueNetIdEOS*>(&NetId);
-	}
+	static const FUniqueNetIdEOS& Cast(const FUniqueNetId& NetId);
 
 	/** global static instance of invalid (zero) id */
-	static const FUniqueNetIdEOSRef& EmptyId()
-	{
-		static const FUniqueNetIdEOSRef EmptyId(Create());
-		return EmptyId;
-	}
+	static const FUniqueNetIdEOSRef& EmptyId();
 
-	static FName GetTypeStatic()
-	{
-		static FName NAME_Eos(TEXT("EOS"));
-		return NAME_Eos;
-	}
+	static FName GetTypeStatic();
 
-	virtual FName GetType() const override
-	{
-		return GetTypeStatic();
-	}
-
-	virtual const uint8* GetBytes() const override
-	{
-		return RawBytes;
-	}
-
-	virtual int32 GetSize() const override
-	{
-		return EOS_ID_BYTE_SIZE;
-	}
+	virtual FName GetType() const override;
+	virtual const uint8* GetBytes() const override;
+	virtual int32 GetSize() const override;
+	virtual bool IsValid() const override;
+	virtual uint32 GetTypeHash() const override;
+	virtual FString ToString() const override;
+	virtual FString ToDebugString() const override;
 
 	const EOS_EpicAccountId GetEpicAccountId() const
 	{
@@ -89,126 +54,42 @@ public:
 		return ProductUserId;
 	}
 
-	void UpdateNetIdStr(const FString& InNetIdStr)
-	{
-		UniqueNetIdStr = InNetIdStr;
-		ParseAccountIds();
-	}
-
-	void ParseAccountIds()
-	{
-		if (ensure(UniqueNetIdStr.Contains(EOS_ID_SEPARATOR)))
-		{
-			TArray<FString> AccountIds;
-			UniqueNetIdStr.ParseIntoArray(AccountIds, EOS_ID_SEPARATOR, false);
-
-			FString EpicAccountIdStr = EMPTY_EASID;
-			if (AccountIds.Num() > 0 && AccountIds[0].Len() > 0)
-			{
-				EpicAccountIdStr = AccountIds[0];
-				EpicAccountId = EOS_EpicAccountId_FromString(TCHAR_TO_UTF8(*EpicAccountIdStr));
-			}
-			else
-			{
-				EpicAccountId = nullptr;
-			}
-			AddToBuffer(RawBytes, EpicAccountIdStr);
-
-			FString ProductUserIdStr = EMPTY_PUID;
-			if (AccountIds.Num() > 1 && AccountIds[1].Len() > 0)
-			{
-				ProductUserIdStr = AccountIds[1];
-				ProductUserId = EOS_ProductUserId_FromString(TCHAR_TO_UTF8(*ProductUserIdStr));
-			}
-			else
-			{
-				ProductUserId = nullptr;
-			}
-			AddToBuffer(RawBytes + ID_HALF_BYTE_SIZE, ProductUserIdStr);
-		}
-		else
-		{
-			UE_LOG_ONLINE(Warning, TEXT("Invalid format for String in FUniqueNetIdEOS constructor. Please use FUniqueNetIdEOS::Create(EOS_EpicAccountId, EOS_ProductUserId) or FUserManagerEOS::MakeNetIdStringFromIds"));
-		}
-	}
-
-	void AddToBuffer(uint8* Buffer, const FString& Source)
-	{
-		check(Source.Len() == 32);
-		for (int32 ReadOffset = 0, WriteOffset = 0; ReadOffset < 32; ReadOffset += 2, WriteOffset++)
-		{
-			FString HexStr = Source.Mid(ReadOffset, 2);
-			// String is in HEX so use the version that takes a base
-			uint8 ToByte = (uint8)FCString::Strtoi(*HexStr, nullptr, 16);
-			Buffer[WriteOffset] = ToByte;
-		}
-	}
-
 private:
-	EOS_EpicAccountId EpicAccountId;
-	EOS_ProductUserId ProductUserId;
+	EOS_EpicAccountId EpicAccountId = nullptr;
+	EOS_ProductUserId ProductUserId = nullptr;
 	uint8 RawBytes[EOS_ID_BYTE_SIZE] = { 0 };
 
-	FUniqueNetIdEOS()
-		PRAGMA_DISABLE_DEPRECATION_WARNINGS
-		: FUniqueNetIdString(EMPTY_EASID EOS_ID_SEPARATOR EMPTY_PUID)
-		PRAGMA_ENABLE_DEPRECATION_WARNINGS
+	friend class FUniqueNetIdEOSRegistry;
+
+	template<typename... TArgs>
+	static FUniqueNetIdEOSRef Create(TArgs&&... Args)
 	{
-		Type = FName("EOS");
+		return MakeShareable(new FUniqueNetIdEOS(Forward<TArgs>(Args)...));
 	}
-	
-	explicit FUniqueNetIdEOS(uint8* Bytes, int32 Size)
-		PRAGMA_DISABLE_DEPRECATION_WARNINGS
-		: FUniqueNetIdString()
-		PRAGMA_ENABLE_DEPRECATION_WARNINGS
-	{
-		check(Size == EOS_ID_BYTE_SIZE);
-		const FString EpicAccountIdStr = BytesToHex(Bytes, ID_HALF_BYTE_SIZE);
-		EpicAccountId = EOS_EpicAccountId_FromString(TCHAR_TO_UTF8(*EpicAccountIdStr));
-		AddToBuffer(RawBytes, EpicAccountIdStr);
-
-		const FString ProductUserIdStr = BytesToHex(Bytes + ID_HALF_BYTE_SIZE, ID_HALF_BYTE_SIZE);
-		ProductUserId = EOS_ProductUserId_FromString(TCHAR_TO_UTF8(*ProductUserIdStr));
-		AddToBuffer(RawBytes + ID_HALF_BYTE_SIZE, ProductUserIdStr);
-
-		UniqueNetIdStr = EpicAccountIdStr + EOS_ID_SEPARATOR + ProductUserIdStr;
 
-		Type = FName("EOS");
-	}
+	FUniqueNetIdEOS() = default;
+	
+	explicit FUniqueNetIdEOS(const uint8* Bytes, int32 Size);
+	explicit FUniqueNetIdEOS(EOS_EpicAccountId InEpicAccountId, EOS_ProductUserId InProductUserId);
+};
 
-	explicit FUniqueNetIdEOS(EOS_EpicAccountId InEpicAccountId, EOS_ProductUserId InProductUserId)
-		PRAGMA_DISABLE_DEPRECATION_WARNINGS
-		: FUniqueNetIdString(MakeNetIdStringFromIds(InEpicAccountId, InProductUserId), FName("EOS"))
-		PRAGMA_ENABLE_DEPRECATION_WARNINGS
-	{
-		ParseAccountIds();
-	}
+class FUniqueNetIdEOSRegistry
+{
+public:
+	static FUniqueNetIdEOSPtr FindOrAdd(const FString& NetIdStr) { return Get().FindOrAddImpl(NetIdStr); }
+	static FUniqueNetIdEOSPtr FindOrAdd(const uint8* Bytes, int32 Size) { return Get().FindOrAddImpl(Bytes, Size); }
+	static FUniqueNetIdEOSPtr FindOrAdd(EOS_EpicAccountId EpicAccountId, EOS_ProductUserId ProductUserId) { return Get().FindOrAddImpl(EpicAccountId, ProductUserId); }
 
-	explicit FUniqueNetIdEOS(const FString& InUniqueNetId)
-		PRAGMA_DISABLE_DEPRECATION_WARNINGS
-		: FUniqueNetIdString(InUniqueNetId, FName("EOS"))
-		PRAGMA_ENABLE_DEPRECATION_WARNINGS
-	{
-		ParseAccountIds();
-	}
+private:
+	FRWLock Lock;
+	TMap<EOS_EpicAccountId, FUniqueNetIdEOSRef> EasToNetId;
+	TMap<EOS_ProductUserId, FUniqueNetIdEOSRef> PuidToNetId;
 
-	explicit FUniqueNetIdEOS(FString&& InUniqueNetId)
-		PRAGMA_DISABLE_DEPRECATION_WARNINGS
-		: FUniqueNetIdString(MoveTemp(InUniqueNetId))
-		PRAGMA_ENABLE_DEPRECATION_WARNINGS
-	{
-		ParseAccountIds();
-		Type = FName("EOS");
-	}
+	static FUniqueNetIdEOSRegistry& Get();
 
-	explicit FUniqueNetIdEOS(const FUniqueNetId& Src)
-		PRAGMA_DISABLE_DEPRECATION_WARNINGS
-		: FUniqueNetIdString(Src)
-		PRAGMA_ENABLE_DEPRECATION_WARNINGS
-	{
-		ParseAccountIds();
-		Type = FName("EOS");
-	}
+	FUniqueNetIdEOSPtr FindOrAddImpl(const FString& NetIdStr);
+	FUniqueNetIdEOSPtr FindOrAddImpl(const uint8* Bytes, int32 Size);
+	FUniqueNetIdEOSPtr FindOrAddImpl(EOS_EpicAccountId EpicAccountId, EOS_ProductUserId ProductUserId);
 };
 
 #ifndef AUTH_ATTR_REFRESH_TOKEN
diff --git a/Engine/Plugins/Online/OnlineSubsystemEOS/Source/OnlineSubsystemEOS/Private/OnlineSubsystemModuleEOS.cpp b/Engine/Plugins/Online/OnlineSubsystemEOS/Source/OnlineSubsystemEOS/Private/OnlineSubsystemModuleEOS.cpp
--- a/Engine/Plugins/Online/OnlineSubsystemEOS/Source/OnlineSubsystemEOS/Private/OnlineSubsystemModuleEOS.cpp
+++ b/Engine/Plugins/Online/OnlineSubsystemEOS/Source/OnlineSubsystemEOS/Private/OnlineSubsystemModuleEOS.cpp
@@ -5,6 +5,7 @@
 #include "OnlineSubsystemModule.h"
 #include "OnlineSubsystemNames.h"
 #include "OnlineSubsystemEOS.h"
+#include "OnlineSubsystemEOSTypes.h"
 #include "EOSSettings.h"
 
 #include "Features/IModularFeature.h"
@@ -12,6 +13,7 @@
 
 #include "Misc/CoreDelegates.h"
 #include "Misc/ConfigCacheIni.h"
+#include "Misc/LazySingleton.h"
 
 #if WITH_EDITOR
 	#include "ISettingsModule.h"
@@ -115,6 +117,8 @@ void FOnlineSubsystemEOSModule::ShutdownModule()
 
 	delete EOSFactory;
 	EOSFactory = nullptr;
+
+	TLazySingleton<FUniqueNetIdEOSRegistry>::TearDown();
 }
 
 #undef LOCTEXT_NAMESPACE
diff --git a/Engine/Plugins/Online/OnlineSubsystemEOS/Source/OnlineSubsystemEOS/Private/UserManagerEOS.cpp b/Engine/Plugins/Online/OnlineSubsystemEOS/Source/OnlineSubsystemEOS/Private/UserManagerEOS.cpp
--- a/Engine/Plugins/Online/OnlineSubsystemEOS/Source/OnlineSubsystemEOS/Private/UserManagerEOS.cpp
+++ b/Engine/Plugins/Online/OnlineSubsystemEOS/Source/OnlineSubsystemEOS/Private/UserManagerEOS.cpp
@@ -1023,8 +1023,8 @@ void FUserManagerEOS::AddLocalUser(int32 LocalUserNum, EOS_EpicAccountId EpicAcc
 		DefaultLocalUser = LocalUserNum;
 	}
 
-	const FString& NetId = MakeNetIdStringFromIds(EpicAccountId, UserId);
-	FUniqueNetIdEOSRef UserNetId = FUniqueNetIdEOS::Create(NetId);
+	FUniqueNetIdEOSRef UserNetId = FUniqueNetIdEOSRegistry::FindOrAdd(EpicAccountId, UserId).ToSharedRef();
+	const FString& NetId = UserNetId->ToString();
 	FUserOnlineAccountEOSRef UserAccountRef(new FUserOnlineAccountEOS(UserNetId));
 
 	UserNumToNetIdMap.Emplace(LocalUserNum, UserNetId);
@@ -1114,7 +1114,7 @@ TSharedPtr<FUserOnlineAccount> FUserManagerEOS::GetUserAccount(const FUniqueNetI
 	TSharedPtr<FUserOnlineAccount> Result;
 
 	const FUniqueNetIdEOS& EOSID = FUniqueNetIdEOS::Cast(UserId);
-	const FUserOnlineAccountEOSRef* FoundUserAccount = StringToUserAccountMap.Find(EOSID.UniqueNetIdStr);
+	const FUserOnlineAccountEOSRef* FoundUserAccount = StringToUserAccountMap.Find(EOSID.ToString());
 	if (FoundUserAccount != nullptr)
 	{
 		return *FoundUserAccount;
@@ -1299,7 +1299,7 @@ void FUserManagerEOS::ResolveUniqueNetIds(const TArray<EOS_ProductUserId>& Produ
 		// We check first if the Product User Id has already been queried, which would allow us to retrieve its Epic Account Id directly
 		if (GetEpicAccountIdFromProductUserId(ProductUserId, EpicAccountId))
 		{
-			const FUniqueNetIdEOSRef UniqueNetId = FUniqueNetIdEOS::Create(MakeNetIdStringFromIds(EpicAccountId, ProductUserId));
+			const FUniqueNetIdEOSRef UniqueNetId = FUniqueNetIdEOSRegistry::FindOrAdd(EpicAccountId, ProductUserId).ToSharedRef();
 
 			ResolvedUniqueNetIds.Add(ProductUserId, UniqueNetId);
 		}
@@ -1332,7 +1332,7 @@ void FUserManagerEOS::ResolveUniqueNetIds(const TArray<EOS_ProductUserId>& Produ
 
 				GetEpicAccountIdFromProductUserId(ProductUserId, EpicAccountId);
 
-				const FUniqueNetIdEOSRef UniqueNetId = FUniqueNetIdEOS::Create(MakeNetIdStringFromIds(EpicAccountId, ProductUserId));
+				const FUniqueNetIdEOSRef UniqueNetId = FUniqueNetIdEOSRegistry::FindOrAdd(EpicAccountId, ProductUserId).ToSharedRef();
 
 				ResolvedUniqueNetIds.Add(ProductUserId, UniqueNetId);
 			}
@@ -1354,9 +1354,9 @@ FOnlineUserPtr FUserManagerEOS::GetLocalOnlineUser(int32 LocalUserNum) const
 	if (UserNumToNetIdMap.Contains(LocalUserNum))
 	{
 		const FUniqueNetIdEOSPtr NetId = UserNumToNetIdMap.FindRef(LocalUserNum);
-		if (NetIdStringToOnlineUserMap.Contains(*NetId->UniqueNetIdStr))
+		if (NetIdStringToOnlineUserMap.Contains(NetId->ToString()))
 		{
-			OnlineUser = NetIdStringToOnlineUserMap.FindRef(*NetId->UniqueNetIdStr);
+			OnlineUser = NetIdStringToOnlineUserMap.FindRef(NetId->ToString());
 		}
 	}
 	return OnlineUser;
@@ -1368,7 +1368,7 @@ FOnlineUserPtr FUserManagerEOS::GetOnlineUser(EOS_ProductUserId UserId) const
 	if (ProductUserIdToStringMap.Contains(UserId))
 	{
 		const FString& NetId = ProductUserIdToStringMap.FindRef(UserId);
-		if (NetIdStringToOnlineUserMap.Contains(*NetId))
+		if (NetIdStringToOnlineUserMap.Contains(NetId))
 		{
 			OnlineUser = NetIdStringToOnlineUserMap.FindRef(*NetId);
 		}
@@ -1382,9 +1382,9 @@ FOnlineUserPtr FUserManagerEOS::GetOnlineUser(EOS_EpicAccountId AccountId) const
 	if (AccountIdToStringMap.Contains(AccountId))
 	{
 		const FString& NetId = AccountIdToStringMap.FindRef(AccountId);
-		if (NetIdStringToOnlineUserMap.Contains(*NetId))
+		if (NetIdStringToOnlineUserMap.Contains(NetId))
 		{
-			OnlineUser = NetIdStringToOnlineUserMap.FindRef(*NetId);
+			OnlineUser = NetIdStringToOnlineUserMap.FindRef(NetId);
 		}
 	}
 	return OnlineUser;
@@ -1397,7 +1397,7 @@ void FUserManagerEOS::RemoveLocalUser(int32 LocalUserNum)
 	{
 		EOSSubsystem->ReleaseVoiceChatUserInterface(**FoundId);
 		LocalUserNumToFriendsListMap.Remove(LocalUserNum);
-		const FString& NetId = (*FoundId)->UniqueNetIdStr;
+		const FString& NetId = (*FoundId)->ToString();
 		const EOS_EpicAccountId AccountId = (*FoundId)->GetEpicAccountId();
 		AccountIdToStringMap.Remove(AccountId);
 		AccountIdToUserNumMap.Remove(AccountId);
@@ -1419,18 +1419,17 @@ void FUserManagerEOS::RemoveLocalUser(int32 LocalUserNum)
 
 FUniqueNetIdPtr FUserManagerEOS::CreateUniquePlayerId(uint8* Bytes, int32 Size)
 {
-	if (Bytes != nullptr && Size >= 32)
-	{
-		// In the case of crossplay the size might be larger, but we only know how to parse the 32
-// @todo joeg crossplay
-		return FUniqueNetIdEOS::Create(Bytes, 32);
-	}
-	return nullptr;
+	// If we're passed an EOSPlus id, the first EOS_ID_BYTE_SIZE bytes are the EAS|EOS part we care about.
+	Size = FMath::Min(Size, EOS_ID_BYTE_SIZE);
+	return FUniqueNetIdEOSRegistry::FindOrAdd(Bytes, Size);
 }
 
-FUniqueNetIdPtr FUserManagerEOS::CreateUniquePlayerId(const FString& Str)
+FUniqueNetIdPtr FUserManagerEOS::CreateUniquePlayerId(const FString& InStr)
 {
-	return FUniqueNetIdEOS::Create(Str);
+	FString NetIdStr = InStr;
+	// If we're passed an EOSPlus id, remove the platform id and separator.
+	NetIdStr.Split(EOSPLUS_ID_SEPARATOR, nullptr, &NetIdStr);
+	return FUniqueNetIdEOSRegistry::FindOrAdd(NetIdStr);
 }
 
 ELoginStatus::Type FUserManagerEOS::GetLoginStatus(int32 LocalUserNum) const
@@ -1748,8 +1747,8 @@ void FUserManagerEOS::FriendStatusChanged(const EOS_Friends_OnFriendsUpdateInfo*
 
 void FUserManagerEOS::AddFriend(int32 LocalUserNum, EOS_EpicAccountId EpicAccountId)
 {
-	const FString NetId = MakeNetIdStringFromIds(EpicAccountId, nullptr);
-	FUniqueNetIdEOSRef FriendNetId = FUniqueNetIdEOS::Create(NetId);
+	FUniqueNetIdEOSRef FriendNetId = FUniqueNetIdEOSRegistry::FindOrAdd(EpicAccountId, nullptr).ToSharedRef();
+	const FString NetId = FriendNetId->ToString();
 	FOnlineFriendEOSRef FriendRef = MakeShareable(new FOnlineFriendEOS(FriendNetId));
 	LocalUserNumToFriendsListMap[LocalUserNum]->Add(NetId, FriendRef);
 
@@ -1775,7 +1774,7 @@ void FUserManagerEOS::AddFriend(int32 LocalUserNum, EOS_EpicAccountId EpicAccoun
 
 void FUserManagerEOS::AddRemotePlayer(int32 LocalUserNum, const FString& NetId, EOS_EpicAccountId EpicAccountId)
 {
-	FUniqueNetIdEOSRef EOSID = FUniqueNetIdEOS::Create(NetId);
+	FUniqueNetIdEOSRef EOSID = FUniqueNetIdEOSRegistry::FindOrAdd(NetId).ToSharedRef();
 	FOnlineUserEOSRef UserRef = MakeShareable(new FOnlineUserEOS(EOSID));
 	// Add this user as a remote (this will grab presence & user info)
 	AddRemotePlayer(LocalUserNum, NetId, EpicAccountId, EOSID, UserRef, UserRef);
@@ -1794,26 +1793,23 @@ void FUserManagerEOS::AddRemotePlayer(int32 LocalUserNum, const FString& NetId,
 	ReadUserInfo(LocalUserNum, EpicAccountId);
 }
 
-void FUserManagerEOS::UpdateRemotePlayerProductUserId(EOS_EpicAccountId AccountId, EOS_ProductUserId UserId)
+void FUserManagerEOS::UpdateRemotePlayerProductUserId(EOS_EpicAccountId EpicAccountId, EOS_ProductUserId ProductUserId)
 {
 	// See if the net ids have changed for this user and bail if they are the same
-	const FString NewNetIdStr = MakeNetIdStringFromIds(AccountId, UserId);
-	const FString PrevNetIdStr = AccountIdToStringMap[AccountId];
+	const FString PrevNetIdStr = AccountIdToStringMap[EpicAccountId];
+	// Note that calling FindOrAdd with previously invalid EAS/EOS components will update the net id in place.
+	const FString NewNetIdStr = FUniqueNetIdEOSRegistry::FindOrAdd(EpicAccountId, ProductUserId)->ToString();
 	if (PrevNetIdStr == NewNetIdStr)
 	{
 		// No change, so skip any work
 		return;
 	}
 
-	const FString UserIdStr = LexToString(UserId);
+	const FString ProductUserIdStr = LexToString(ProductUserId);
 
 	// Get the unique net id and rebuild the string for it
 	IAttributeAccessInterfaceRef AttrAccess = NetIdStringToAttributeAccessMap[PrevNetIdStr];
 	FUniqueNetIdEOSPtr NetIdEOS = AttrAccess->GetUniqueNetIdEOS();
-	if (NetIdEOS.IsValid())
-	{
-		ConstCastSharedPtr<FUniqueNetIdEOS>(NetIdEOS)->UpdateNetIdStr(NewNetIdStr);
-	}
 	// Update any old friends entries with the new net id key
 	for (TMap<int32, FFriendsListEOSRef>::TConstIterator It(LocalUserNumToFriendsListMap); It; ++It)
 	{
@@ -1825,10 +1821,10 @@ void FUserManagerEOS::UpdateRemotePlayerProductUserId(EOS_EpicAccountId AccountI
 		}
 	}
 	// Update all of the other net id to X mappings
-	AccountIdToStringMap.Remove(AccountId);
-	AccountIdToStringMap.Emplace(AccountId, NewNetIdStr);
-	ProductUserIdToStringMap.Remove(UserId);
-	ProductUserIdToStringMap.Emplace(UserId, *NewNetIdStr);
+	AccountIdToStringMap.Remove(EpicAccountId);
+	AccountIdToStringMap.Emplace(EpicAccountId, NewNetIdStr);
+	ProductUserIdToStringMap.Remove(ProductUserId);
+	ProductUserIdToStringMap.Emplace(ProductUserId, *NewNetIdStr);
 	FOnlineUserPtr OnlineUser = NetIdStringToOnlineUserMap[PrevNetIdStr];
 	NetIdStringToOnlineUserMap.Remove(PrevNetIdStr);
 	NetIdStringToOnlineUserMap.Emplace(NewNetIdStr, OnlineUser);
@@ -2017,7 +2013,7 @@ bool FUserManagerEOS::SendInvite(int32 LocalUserNum, const FUniqueNetId& FriendI
 		{
 			ErrorString = FString::Printf(TEXT("Failed to send invite for user (%d) to player (%s) with result code (%s)"), LocalUserNum, *NetId, ANSI_TO_TCHAR(EOS_EResult_ToString(Data->ResultCode)));
 		}
-		Delegate.ExecuteIfBound(LocalUserNum, bWasSuccessful, *FUniqueNetIdEOS::Create(NetId), ListName, ErrorString);
+		Delegate.ExecuteIfBound(LocalUserNum, bWasSuccessful, *FUniqueNetIdEOSRegistry::FindOrAdd(NetId), ListName, ErrorString);
 	};
 
 	EOS_Friends_SendInviteOptions Options = { };
@@ -2060,7 +2056,7 @@ bool FUserManagerEOS::AcceptInvite(int32 LocalUserNum, const FUniqueNetId& Frien
 		{
 			ErrorString = FString::Printf(TEXT("Failed to accept invite for user (%d) from friend (%s) with result code (%s)"), LocalUserNum, *NetId, ANSI_TO_TCHAR(EOS_EResult_ToString(Data->ResultCode)));
 		}
-		Delegate.ExecuteIfBound(LocalUserNum, bWasSuccessful, *FUniqueNetIdEOS::Create(NetId), ListName, ErrorString);
+		Delegate.ExecuteIfBound(LocalUserNum, bWasSuccessful, *FUniqueNetIdEOSRegistry::FindOrAdd(NetId), ListName, ErrorString);
 	};
 
 	EOS_Friends_AcceptInviteOptions Options = { };
@@ -2173,7 +2169,7 @@ TSharedPtr<FOnlineFriend> FUserManagerEOS::GetFriend(int32 LocalUserNum, const F
 	{
 		FFriendsListEOSRef FriendsList = LocalUserNumToFriendsListMap[LocalUserNum];
 		const FUniqueNetIdEOS& EosId = FUniqueNetIdEOS::Cast(FriendId);
-		FOnlineFriendEOSPtr FoundFriend = FriendsList->GetByNetIdString(EosId.UniqueNetIdStr);
+		FOnlineFriendEOSPtr FoundFriend = FriendsList->GetByNetIdString(EosId.ToString());
 		if (FoundFriend.IsValid())
 		{
 			const FOnlineUserPresence& Presence = FoundFriend->GetPresence();
@@ -2317,7 +2313,7 @@ bool FUserManagerEOS::HandleFriendsExec(UWorld* InWorld, const TCHAR* Cmd, FOutp
 		int LocalUserNum = FCString::Atoi(*FParse::Token(Cmd, false));
 
 		FString FriendUserIdStr = FParse::Token(Cmd, false);
-		FUniqueNetIdEOSRef FriendEosId = FUniqueNetIdEOS::Create(FriendUserIdStr);
+		FUniqueNetIdEOSRef FriendEosId = FUniqueNetIdEOSRegistry::FindOrAdd(FriendUserIdStr).ToSharedRef();
 
 		FString FriendsList = FParse::Token(Cmd, false);
 
@@ -2341,7 +2337,7 @@ bool FUserManagerEOS::HandleFriendsExec(UWorld* InWorld, const TCHAR* Cmd, FOutp
 		int LocalUserNum = FCString::Atoi(*FParse::Token(Cmd, false));
 
 		FString FriendUserIdStr = FParse::Token(Cmd, false);
-		FUniqueNetIdEOSRef FriendEosId = FUniqueNetIdEOS::Create(FriendUserIdStr);
+		FUniqueNetIdEOSRef FriendEosId = FUniqueNetIdEOSRegistry::FindOrAdd(FriendUserIdStr).ToSharedRef();
 
 		FString FriendsList = FParse::Token(Cmd, false);
 
@@ -2392,7 +2388,7 @@ void FUserManagerEOS::SetPresence(const FUniqueNetId& UserId, const FOnlineUserP
 	const EOS_EpicAccountId AccountId = EOSID.GetEpicAccountId();
 	if (EOS_EpicAccountId_IsValid(AccountId) == EOS_FALSE)
 	{
-		UE_LOG_ONLINE(Error, TEXT("Can't SetPresence() for user (%s) since they are not logged in"), *EOSID.UniqueNetIdStr);
+		UE_LOG_ONLINE(Error, TEXT("Can't SetPresence() for user (%s) since they are not logged in"), *EOSID.ToDebugString());
 		return;
 	}
 
@@ -2445,7 +2441,7 @@ void FUserManagerEOS::SetPresence(const FUniqueNetId& UserId, const FOnlineUserP
 	EOS_EResult SetDataResult = EOS_PresenceModification_SetData(ChangeHandle, &DataOptions);
 	if (SetDataResult != EOS_EResult::EOS_Success)
 	{
-		UE_LOG_ONLINE(Error, TEXT("EOS_PresenceModification_SetData() failed with result code (%s)"), ANSI_TO_TCHAR(EOS_EResult_ToString(SetDataResult)));
+		UE_LOG_ONLINE(Error, TEXT("EOS_PresenceModification_SetData() failed with result code (%s)"), *LexToString(SetDataResult));
 	}
 
 	FSetPresenceCallback* CallbackObj = new FSetPresenceCallback(FUserManagerEOSWeakPtr(AsShared()));
@@ -2453,11 +2449,11 @@ void FUserManagerEOS::SetPresence(const FUniqueNetId& UserId, const FOnlineUserP
 	{
 		if (Data->ResultCode == EOS_EResult::EOS_Success && AccountIdToStringMap.Contains(Data->LocalUserId))
 		{
-			FUniqueNetIdEOSRef EOSID = FUniqueNetIdEOS::Create(AccountIdToStringMap[Data->LocalUserId]);
+			FUniqueNetIdEOSRef EOSID = FUniqueNetIdEOSRegistry::FindOrAdd(AccountIdToStringMap[Data->LocalUserId]).ToSharedRef();
 			Delegate.ExecuteIfBound(*EOSID, true);
 			return;
 		}
-		UE_LOG_ONLINE(Error, TEXT("SetPresence() failed with result code (%s)"), ANSI_TO_TCHAR(EOS_EResult_ToString(Data->ResultCode)));
+		UE_LOG_ONLINE(Error, TEXT("SetPresence() failed with result code (%s)"), *LexToString(Data->ResultCode));
 		Delegate.ExecuteIfBound(*FUniqueNetIdEOS::EmptyId(), false);
 	};
 
@@ -2485,7 +2481,7 @@ void FUserManagerEOS::QueryPresence(const FUniqueNetId& UserId, const FOnPresenc
 	const EOS_EpicAccountId AccountId = EOSID.GetEpicAccountId();
 	if (EOS_EpicAccountId_IsValid(AccountId) == EOS_FALSE)
 	{
-		UE_LOG_ONLINE(Error, TEXT("Can't QueryPresence(%s) for unknown unique net id"), *EOSID.UniqueNetIdStr);
+		UE_LOG_ONLINE(Error, TEXT("Can't QueryPresence(%s) for unknown unique net id"), *EOSID.ToDebugString());
 		Delegate.ExecuteIfBound(UserId, false);
 		return;
 	}
@@ -2508,8 +2504,8 @@ void FUserManagerEOS::QueryPresence(const FUniqueNetId& UserId, const FOnPresenc
 				Delegate.ExecuteIfBound(*OnlineUser->GetUserId(), true);
 				return;
 			}
-			const FString& TargetUser = MakeNetIdStringFromIds(Data->TargetUserId, nullptr);
-			UE_LOG_ONLINE(Error, TEXT("QueryPresence() for user (%s) failed with result code (%s)"), *TargetUser, ANSI_TO_TCHAR(EOS_EResult_ToString(Data->ResultCode)));
+			const FString& TargetUser = FUniqueNetIdEOSRegistry::FindOrAdd(Data->TargetUserId, nullptr)->ToString();
+			UE_LOG_ONLINE(Error, TEXT("QueryPresence() for user (%s) failed with result code (%s)"), *TargetUser, *LexToString(Data->ResultCode));
 			Delegate.ExecuteIfBound(*FUniqueNetIdEOS::EmptyId(), false);
 		};
 
@@ -2594,9 +2590,9 @@ void FUserManagerEOS::UpdateFriendPresence(const FString& FriendId, FOnlineUserP
 EOnlineCachedResult::Type FUserManagerEOS::GetCachedPresence(const FUniqueNetId& UserId, TSharedPtr<FOnlineUserPresence>& OutPresence)
 {
 	const FUniqueNetIdEOS& EOSID = FUniqueNetIdEOS::Cast(UserId);
-	if (NetIdStringToOnlineUserPresenceMap.Contains(EOSID.UniqueNetIdStr))
+	if (NetIdStringToOnlineUserPresenceMap.Contains(EOSID.ToString()))
 	{
-		OutPresence = NetIdStringToOnlineUserPresenceMap[EOSID.UniqueNetIdStr];
+		OutPresence = NetIdStringToOnlineUserPresenceMap[EOSID.ToString()];
 		return EOnlineCachedResult::Success;
 	}
 	return EOnlineCachedResult::NotFound;
@@ -2617,7 +2613,7 @@ bool FUserManagerEOS::QueryUserInfo(int32 LocalUserNum, const TArray<FUniqueNetI
 	{
 		const FUniqueNetIdEOS& EOSID = FUniqueNetIdEOS::Cast(*NetId);
 		// Skip querying for local users since we already have that data
-		if (StringToUserAccountMap.Contains(EOSID.UniqueNetIdStr))
+		if (StringToUserAccountMap.Contains(EOSID.ToString()))
 		{
 			continue;
 		}
@@ -2636,7 +2632,7 @@ bool FUserManagerEOS::QueryUserInfo(int32 LocalUserNum, const TArray<FUniqueNetI
 				UserEasIdsNeedingExternalMappings.Add(LexToString(AccountId));
 
 				// Registering the player will also query the user info data
-				AddRemotePlayer(LocalUserNum, EOSID.UniqueNetIdStr, AccountId);
+				AddRemotePlayer(LocalUserNum, EOSID.ToString(), AccountId);
 			}
 		}
 	}
@@ -2699,9 +2695,9 @@ TSharedPtr<FOnlineUser> FUserManagerEOS::GetUserInfo(int32 LocalUserNum, const F
 {
 	TSharedPtr<FOnlineUser> OnlineUser;
 	const FUniqueNetIdEOS& EOSID = FUniqueNetIdEOS::Cast(UserId);
-	if (NetIdStringToOnlineUserMap.Contains(EOSID.UniqueNetIdStr))
+	if (NetIdStringToOnlineUserMap.Contains(EOSID.ToString()))
 	{
-		OnlineUser = NetIdStringToOnlineUserMap[EOSID.UniqueNetIdStr];
+		OnlineUser = NetIdStringToOnlineUserMap[EOSID.ToString()];
 	}
 	return OnlineUser;
 }
@@ -2726,8 +2722,8 @@ bool FUserManagerEOS::QueryUserIdMapping(const FUniqueNetId& UserId, const FStri
 	const EOS_EpicAccountId AccountId = EOSID.GetEpicAccountId();
 	if (EOS_EpicAccountId_IsValid(AccountId) == EOS_FALSE)
 	{
-		UE_LOG_ONLINE(Error, TEXT("Specified local user (%s) is not known"), *EOSID.UniqueNetIdStr);
-		Delegate.ExecuteIfBound(false, UserId, DisplayNameOrEmail, *FUniqueNetIdEOS::EmptyId(), FString::Printf(TEXT("Specified local user (%s) is not known"), *EOSID.UniqueNetIdStr));
+		UE_LOG_ONLINE(Error, TEXT("Specified local user (%s) is not known"), *EOSID.ToDebugString());
+		Delegate.ExecuteIfBound(false, UserId, DisplayNameOrEmail, *FUniqueNetIdEOS::EmptyId(), FString::Printf(TEXT("Specified local user (%s) is not known"), *EOSID.ToDebugString()));
 		return false;
 	}
 	int32 LocalUserNum = GetLocalUserNumFromUniqueNetId(UserId);
@@ -2754,7 +2750,7 @@ bool FUserManagerEOS::QueryUserIdMapping(const FUniqueNetId& UserId, const FStri
 				AddRemotePlayer(LocalUserNum, NetIdStr, Data->TargetUserId);
 			}
 
-			Delegate.ExecuteIfBound(true, *LocalUserId, DisplayNameOrEmail, *FUniqueNetIdEOS::Create(NetIdStr), ErrorString);
+			Delegate.ExecuteIfBound(true, *LocalUserId, DisplayNameOrEmail, *FUniqueNetIdEOSRegistry::FindOrAdd(NetIdStr), ErrorString);
 		}
 		else
 		{
@@ -2820,13 +2816,13 @@ bool FUserManagerEOS::QueryExternalIdMappings(const FUniqueNetId& UserId, const
 	const EOS_EpicAccountId AccountId = EOSID.GetEpicAccountId();
 	if (EOS_EpicAccountId_IsValid(AccountId) == EOS_FALSE)
 	{
-		Delegate.ExecuteIfBound(false, UserId, QueryOptions, ExternalIds, FString::Printf(TEXT("User (%s) is not logged in, so can't query external account ids"), *EOSID.UniqueNetIdStr));
+		Delegate.ExecuteIfBound(false, UserId, QueryOptions, ExternalIds, FString::Printf(TEXT("User (%s) is not logged in, so can't query external account ids"), *EOSID.ToDebugString()));
 		return false;
 	}
 
 	if (ExternalIds.IsEmpty())
 	{
-		Delegate.ExecuteIfBound(false, UserId, QueryOptions, ExternalIds, FString::Printf(TEXT("List of ids to query is empty for User (%s), so can't query external account ids"), *EOSID.UniqueNetIdStr));
+		Delegate.ExecuteIfBound(false, UserId, QueryOptions, ExternalIds, FString::Printf(TEXT("List of ids to query is empty for User (%s), so can't query external account ids"), *EOSID.ToDebugString()));
 		return false;
 	}
 
diff --git a/Engine/Plugins/Online/OnlineSubsystemEOS/Source/OnlineSubsystemEOS/Private/UserManagerEOS.h b/Engine/Plugins/Online/OnlineSubsystemEOS/Source/OnlineSubsystemEOS/Private/UserManagerEOS.h
--- a/Engine/Plugins/Online/OnlineSubsystemEOS/Source/OnlineSubsystemEOS/Private/UserManagerEOS.h
+++ b/Engine/Plugins/Online/OnlineSubsystemEOS/Source/OnlineSubsystemEOS/Private/UserManagerEOS.h
@@ -4,6 +4,8 @@
 
 #include "OnlineSubsystemEOSTypes.h"
 #include "Interfaces/OnlineExternalUIInterface.h"
+#include "Interfaces/OnlineFriendsInterface.h"
+#include "Interfaces/OnlinePresenceInterface.h"
 
 #if WITH_EOS_SDK
 	#include "eos_auth_types.h"
diff --git a/Engine/Plugins/Online/OnlineSubsystemEOS/Source/OnlineSubsystemEOSPlus/Private/OnlineUserEOSPlus.cpp b/Engine/Plugins/Online/OnlineSubsystemEOS/Source/OnlineSubsystemEOSPlus/Private/OnlineUserEOSPlus.cpp
--- a/Engine/Plugins/Online/OnlineSubsystemEOS/Source/OnlineSubsystemEOSPlus/Private/OnlineUserEOSPlus.cpp
+++ b/Engine/Plugins/Online/OnlineSubsystemEOS/Source/OnlineSubsystemEOSPlus/Private/OnlineUserEOSPlus.cpp
@@ -5,6 +5,8 @@
 #include "OnlineSubsystemEOSPlus.h"
 #include "EOSSettings.h"
 
+#define EOSPLUS_ID_SEPARATOR TEXT("_+_")
+
 enum class EOSSValue : uint8
 {
 	Null,
@@ -72,7 +74,7 @@ PRAGMA_ENABLE_DEPRECATION_WARNINGS
 inline FString BuildEOSPlusStringId(FUniqueNetIdPtr InBaseUniqueNetId, FUniqueNetIdPtr InEOSUniqueNetId)
 {
 	FString StrId = InBaseUniqueNetId.IsValid() ? InBaseUniqueNetId->ToString() : TEXT("");
-	StrId += TEXT("_+_");
+	StrId += EOSPLUS_ID_SEPARATOR;
 	StrId += InEOSUniqueNetId.IsValid() ? InEOSUniqueNetId->ToString() : TEXT("");
 	return StrId;
 }
@@ -711,28 +713,29 @@ FUniqueNetIdPtr FOnlineUserEOSPlus::CreateUniquePlayerId(uint8* Bytes, int32 Siz
 		UE_LOG_ONLINE(Error, TEXT("Invalid size (%d) passed to FOnlineUserEOSPlus::CreateUniquePlayerId()"), Size);
 		return nullptr;
 	}
-	// We know that the first EOS_NETID_BYTE_SIZE bytes are the EOS ids, so the rest is the platform id
+
+	// We know that the first EOS_NETID_BYTE_SIZE bytes are the EOS ids
 	FUniqueNetIdPtr EOSNetId = EOSIdentityInterface->CreateUniquePlayerId(Bytes, EOS_NETID_BYTE_SIZE);
+
+	// The rest is the platform id, which might be missing if the passed bytes is from an EOS id.
 	int32 BaseByteOffset = EOS_NETID_BYTE_SIZE;
 	int32 PlatformIdSize = Size - BaseByteOffset - BASE_NETID_TYPE_SIZE;
-	if (PlatformIdSize < 0)
-	{
-		UE_LOG_ONLINE(Error, TEXT("Invalid size (%d) passed to FOnlineUserEOSPlus::CreateUniquePlayerId()"), Size);
-		return nullptr;
-	}
-	uint8 OSSType = *(Bytes + BaseByteOffset);
-	BaseByteOffset += BASE_NETID_TYPE_SIZE;
-	FName OSSName = ToOSSName((EOSSValue)OSSType);
-	FName BaseOSSName = EOSPlus->BaseOSS->GetSubsystemName();
 	FUniqueNetIdPtr BaseNetId = nullptr;
-	if (BaseOSSName == OSSName)
-	{
-		BaseNetId = BaseIdentityInterface->CreateUniquePlayerId(Bytes + BaseByteOffset, PlatformIdSize);
-	}
-	else
-	{
-		// Just create the pass through version that holds the other platform data but doesn't interpret it
-		BaseNetId = FUniqueNetIdBinary::Create(Bytes + BaseByteOffset, PlatformIdSize, OSSName);
+	if (PlatformIdSize > 0)
+	{	
+		uint8 OSSType = *(Bytes + BaseByteOffset);
+		BaseByteOffset += BASE_NETID_TYPE_SIZE;
+		FName OSSName = ToOSSName((EOSSValue)OSSType);
+		FName BaseOSSName = EOSPlus->BaseOSS->GetSubsystemName();
+		if (BaseOSSName == OSSName)
+		{
+			BaseNetId = BaseIdentityInterface->CreateUniquePlayerId(Bytes + BaseByteOffset, PlatformIdSize);
+		}
+		else
+		{
+			// Just create the pass through version that holds the other platform data but doesn't interpret it
+			BaseNetId = FUniqueNetIdBinary::Create(Bytes + BaseByteOffset, PlatformIdSize, OSSName);
+		}
 	}
 	
 	return AddRemotePlayer(BaseNetId, EOSNetId);
@@ -740,17 +743,17 @@ FUniqueNetIdPtr FOnlineUserEOSPlus::CreateUniquePlayerId(uint8* Bytes, int32 Siz
 
 FUniqueNetIdPtr FOnlineUserEOSPlus::CreateUniquePlayerId(const FString& Str)
 {
-	// Split <id>_+_<id2> into two strings
-	int32 FoundAt = Str.Find(TEXT("_+_"));
-	if (FoundAt == -1)
+	// Split <platformid>_+_<eas|eos> into two strings
+	FString BaseNetIdStr;
+	FString EOSNetIdStr;
+	if (!Str.Split(EOSPLUS_ID_SEPARATOR, &BaseNetIdStr, &EOSNetIdStr))
 	{
-		UE_LOG_ONLINE(Error, TEXT("Couldn't parse string (%s) passed to FOnlineUserEOSPlus::CreateUniquePlayerId()"), *Str);
-		return nullptr;
+		// If we couldn't find the EOSPlus separator, this must be an EOS net id
+		EOSNetIdStr = Str;
 	}
 
 	// BaseNetId will only be valid for platform friends
 	FUniqueNetIdPtr BaseNetId = nullptr;
-	FString BaseNetIdStr = Str.Left(FoundAt);
 	if (!BaseNetIdStr.IsEmpty())
 	{
 		BaseNetId = BaseIdentityInterface->CreateUniquePlayerId(BaseNetIdStr);
@@ -758,7 +761,6 @@ FUniqueNetIdPtr FOnlineUserEOSPlus::CreateUniquePlayerId(const FString& Str)
 
 	// EOSNetId will only be valid for EOS friends
 	FUniqueNetIdPtr EOSNetId = nullptr;
-	FString EOSNetIdStr = Str.Right(Str.Len() - FoundAt - 3);
 	if (!EOSNetIdStr.IsEmpty())
 	{
 		EOSNetId = EOSIdentityInterface->CreateUniquePlayerId(EOSNetIdStr);
diff --git a/Engine/Plugins/Online/OnlineSubsystemUtils/Source/OnlineSubsystemUtils/Private/OnlineEngineInterfaceImpl.cpp b/Engine/Plugins/Online/OnlineSubsystemUtils/Source/OnlineSubsystemUtils/Private/OnlineEngineInterfaceImpl.cpp
--- a/Engine/Plugins/Online/OnlineSubsystemUtils/Source/OnlineSubsystemUtils/Private/OnlineEngineInterfaceImpl.cpp
+++ b/Engine/Plugins/Online/OnlineSubsystemUtils/Source/OnlineSubsystemUtils/Private/OnlineEngineInterfaceImpl.cpp
@@ -123,7 +123,7 @@ FUniqueNetIdWrapper UOnlineEngineInterfaceImpl::CreateUniquePlayerIdWrapper(cons
 	{
 		// No UWorld here, but ok since this is just a factory
 		UWorld* World = nullptr;
-		IOnlineIdentityPtr IdentityInt = Online::GetIdentityInterface(World, bIsPrimaryLoaded ? Type : *MappedUniqueNetIdType);
+		IOnlineIdentityPtr IdentityInt = Online::GetIdentityInterface(World, bIsMappedUniqueNetIdTypeLoaded ? *MappedUniqueNetIdType : Type);
 		if (IdentityInt.IsValid())
 		{
 			UniqueId = IdentityInt->CreateUniquePlayerId(Str);
