Addition of logic to make FUserManagerEOS robust to callbacks returning after module shutdown.

diff --git a/Engine/Plugins/Online/OnlineSubsystemEOS/Source/OnlineSubsystemEOS/Private/UserManagerEOS.cpp b/Engine/Plugins/Online/OnlineSubsystemEOS/Source/OnlineSubsystemEOS/Private/UserManagerEOS.cpp
--- a/Engine/Plugins/Online/OnlineSubsystemEOS/Source/OnlineSubsystemEOS/Private/UserManagerEOS.cpp
+++ b/Engine/Plugins/Online/OnlineSubsystemEOS/Source/OnlineSubsystemEOS/Private/UserManagerEOS.cpp
@@ -158,12 +158,12 @@ void FUserManagerEOS::Init()
 		EOS_UI_AddNotifyDisplaySettingsUpdatedOptions Options = {};
 		Options.ApiVersion = EOS_UI_ADDNOTIFYDISPLAYSETTINGSUPDATED_API_LATEST;
 
-	FOnDisplaySettingsUpdatedCallback* CallbackObj = new FOnDisplaySettingsUpdatedCallback(FUserManagerEOSWeakPtr(AsShared()));
-	DisplaySettingsUpdatedCallback = CallbackObj;
-	CallbackObj->CallbackLambda = [this](const EOS_UI_OnDisplaySettingsUpdatedCallbackInfo* Data)
-	{
-		TriggerOnExternalUIChangeDelegates((bool)Data->bIsVisible);
-	};
+		FOnDisplaySettingsUpdatedCallback* CallbackObj = new FOnDisplaySettingsUpdatedCallback(AsWeak());
+		DisplaySettingsUpdatedCallback = CallbackObj;
+		CallbackObj->CallbackLambda = [this](const EOS_UI_OnDisplaySettingsUpdatedCallbackInfo* Data)
+		{
+			TriggerOnExternalUIChangeDelegates((bool)Data->bIsVisible);
+		};
 
 		DisplaySettingsUpdatedId = EOS_UI_AddNotifyDisplaySettingsUpdated(EOSSubsystem->UIHandle, &Options, CallbackObj, CallbackObj->GetCallbackPtr());
 	}
@@ -477,7 +477,7 @@ bool FUserManagerEOS::Login(int32 LocalUserNum, const FOnlineAccountCredentials&
 		return false;
 	}
 
-	FLoginCallback* CallbackObj = new FLoginCallback(FUserManagerEOSWeakPtr(AsShared()));
+	FLoginCallback* CallbackObj = new FLoginCallback(AsWeak());
 	CallbackObj->CallbackLambda = [this, LocalUserNum, bIsPersistentLogin](const EOS_Auth_LoginCallbackInfo* Data)
 	{
 		if (Data->ResultCode == EOS_EResult::EOS_Success)
@@ -502,7 +502,7 @@ bool FUserManagerEOS::Login(int32 LocalUserNum, const FOnlineAccountCredentials&
 			// Check for invalid persistent login credentials.
 			if (bIsPersistentLogin && bShouldRemoveCachedToken)
 			{
-				FDeletePersistentAuthCallback* DeleteAuthCallbackObj = new FDeletePersistentAuthCallback(FUserManagerEOSWeakPtr(AsShared()));
+				FDeletePersistentAuthCallback* DeleteAuthCallbackObj = new FDeletePersistentAuthCallback(AsWeak());
 				DeleteAuthCallbackObj->CallbackLambda = [this, LocalUserNum, TriggerLoginFailure](const EOS_Auth_DeletePersistentAuthCallbackInfo* Data)
 				{
 					// Deleting the auth token is best effort.
@@ -528,44 +528,47 @@ bool FUserManagerEOS::Login(int32 LocalUserNum, const FOnlineAccountCredentials&
 void FUserManagerEOS::LoginViaExternalAuth(int32 LocalUserNum)
 {
 	GetPlatformAuthToken(LocalUserNum,
-		FOnGetLinkedAccountAuthTokenCompleteDelegate::CreateLambda([this](int32 LocalUserNum, bool bWasSuccessful, const FExternalAuthToken& AuthToken)
+		FOnGetLinkedAccountAuthTokenCompleteDelegate::CreateLambda([this, WeakThis = AsWeak()](int32 LocalUserNum, bool bWasSuccessful, const FExternalAuthToken& AuthToken)
 		{
-			if (!bWasSuccessful || !AuthToken.IsValid())
+			if (FUserManagerEOSPtr StrongThis = WeakThis.Pin())
 			{
-				UE_LOG_ONLINE(Warning, TEXT("Unable to Login() user (%d) due to an empty platform auth token"), LocalUserNum);
-				TriggerOnLoginCompleteDelegates(LocalUserNum, false, *FUniqueNetIdEOS::EmptyId(), FString(TEXT("Missing platform auth token")));
-				return;
-			}
+				if (!bWasSuccessful || !AuthToken.IsValid())
+				{
+					UE_LOG_ONLINE(Warning, TEXT("Unable to Login() user (%d) due to an empty platform auth token"), LocalUserNum);
+					TriggerOnLoginCompleteDelegates(LocalUserNum, false, *FUniqueNetIdEOS::EmptyId(), FString(TEXT("Missing platform auth token")));
+					return;
+				}
 
-			EOS_Auth_LoginOptions LoginOptions = { };
-			LoginOptions.ApiVersion = EOS_AUTH_LOGIN_API_LATEST;
-			LoginOptions.ScopeFlags = EOS_EAuthScopeFlags::EOS_AS_BasicProfile | EOS_EAuthScopeFlags::EOS_AS_FriendsList | EOS_EAuthScopeFlags::EOS_AS_Presence;
+				EOS_Auth_LoginOptions LoginOptions = { };
+				LoginOptions.ApiVersion = EOS_AUTH_LOGIN_API_LATEST;
+				LoginOptions.ScopeFlags = EOS_EAuthScopeFlags::EOS_AS_BasicProfile | EOS_EAuthScopeFlags::EOS_AS_FriendsList | EOS_EAuthScopeFlags::EOS_AS_Presence;
 
-			check(LocalUserNumToLastLoginCredentials.Contains(LocalUserNum));
-			FAuthCredentials Credentials(ToEOS_EExternalCredentialType(GetPlatformOSS()->GetSubsystemName(), *LocalUserNumToLastLoginCredentials[LocalUserNum]), AuthToken);
-			LoginOptions.Credentials = &Credentials;
+				check(LocalUserNumToLastLoginCredentials.Contains(LocalUserNum));
+				FAuthCredentials Credentials(ToEOS_EExternalCredentialType(GetPlatformOSS()->GetSubsystemName(), *LocalUserNumToLastLoginCredentials[LocalUserNum]), AuthToken);
+				LoginOptions.Credentials = &Credentials;
 
-			FLoginCallback* CallbackObj = new FLoginCallback(FUserManagerEOSWeakPtr(AsShared()));
-			CallbackObj->CallbackLambda = [this, LocalUserNum](const EOS_Auth_LoginCallbackInfo* Data)
-			{
-				if (Data->ResultCode == EOS_EResult::EOS_Success)
-				{
-					ConnectLoginEAS(LocalUserNum, Data->LocalUserId);
-				}
-				else if (Data->ResultCode == EOS_EResult::EOS_InvalidUser)
-				{
-					// Link the account
-					LinkEAS(LocalUserNum, Data->ContinuanceToken);
-				}
-				else
+				FLoginCallback* CallbackObj = new FLoginCallback(AsWeak());
+				CallbackObj->CallbackLambda = [this, LocalUserNum](const EOS_Auth_LoginCallbackInfo* Data)
 				{
-					FString ErrorString = FString::Printf(TEXT("Login(%d) failed with EOS result code (%s)"), LocalUserNum, ANSI_TO_TCHAR(EOS_EResult_ToString(Data->ResultCode)));
-					UE_LOG_ONLINE(Warning, TEXT("%s"), *ErrorString);
-					TriggerOnLoginCompleteDelegates(LocalUserNum, false, *FUniqueNetIdEOS::EmptyId(), ErrorString);
-				}
-			};
-			// Perform the auth call
-			EOS_Auth_Login(EOSSubsystem->AuthHandle, &LoginOptions, (void*)CallbackObj, CallbackObj->GetCallbackPtr());
+					if (Data->ResultCode == EOS_EResult::EOS_Success)
+					{
+						ConnectLoginEAS(LocalUserNum, Data->LocalUserId);
+					}
+					else if (Data->ResultCode == EOS_EResult::EOS_InvalidUser)
+					{
+						// Link the account
+						LinkEAS(LocalUserNum, Data->ContinuanceToken);
+					}
+					else
+					{
+						FString ErrorString = FString::Printf(TEXT("Login(%d) failed with EOS result code (%s)"), LocalUserNum, ANSI_TO_TCHAR(EOS_EResult_ToString(Data->ResultCode)));
+						UE_LOG_ONLINE(Warning, TEXT("%s"), *ErrorString);
+						TriggerOnLoginCompleteDelegates(LocalUserNum, false, *FUniqueNetIdEOS::EmptyId(), ErrorString);
+					}
+				};
+				// Perform the auth call
+				EOS_Auth_Login(EOSSubsystem->AuthHandle, &LoginOptions, (void*)CallbackObj, CallbackObj->GetCallbackPtr());
+			}
 		}));
 }
 
@@ -585,7 +588,7 @@ typedef TEOSCallback<EOS_Auth_OnLinkAccountCallback, EOS_Auth_LinkAccountCallbac
 void FUserManagerEOS::LinkEAS(int32 LocalUserNum, EOS_ContinuanceToken Token)
 {
 	FLinkAccountOptions Options(Token);
-	FLinkAccountCallback* CallbackObj = new FLinkAccountCallback(FUserManagerEOSWeakPtr(AsShared()));
+	FLinkAccountCallback* CallbackObj = new FLinkAccountCallback(AsWeak());
 	CallbackObj->CallbackLambda = [this, LocalUserNum](const EOS_Auth_LinkAccountCallbackInfo* Data)
 	{
 		if (Data->ResultCode == EOS_EResult::EOS_Success)
@@ -648,53 +651,56 @@ struct FConnectCredentials :
 bool FUserManagerEOS::ConnectLoginNoEAS(int32 LocalUserNum)
 {
 	GetPlatformAuthToken(LocalUserNum,
-		FOnGetLinkedAccountAuthTokenCompleteDelegate::CreateLambda([this](int32 LocalUserNum, bool bWasSuccessful, const FExternalAuthToken& AuthToken)
+		FOnGetLinkedAccountAuthTokenCompleteDelegate::CreateLambda([this, WeakThis = AsWeak()](int32 LocalUserNum, bool bWasSuccessful, const FExternalAuthToken& AuthToken)
 		{
-			if (!bWasSuccessful || !AuthToken.IsValid())
+			if (FUserManagerEOSPtr StrongThis = WeakThis.Pin())
 			{
-				const FString ErrorString = FString::Printf(TEXT("ConnectLoginNoEAS(%d) failed due to the platform OSS giving an empty auth token"), LocalUserNum);
-				UE_LOG_ONLINE(Warning, TEXT("%s"), *ErrorString);
-				TriggerOnLoginCompleteDelegates(LocalUserNum, false, *FUniqueNetIdEOS::EmptyId(), ErrorString);
-				return;
-			}
+				if (!bWasSuccessful || !AuthToken.IsValid())
+				{
+					const FString ErrorString = FString::Printf(TEXT("ConnectLoginNoEAS(%d) failed due to the platform OSS giving an empty auth token"), LocalUserNum);
+					UE_LOG_ONLINE(Warning, TEXT("%s"), *ErrorString);
+					TriggerOnLoginCompleteDelegates(LocalUserNum, false, *FUniqueNetIdEOS::EmptyId(), ErrorString);
+					return;
+				}
 
-			// Now login into our EOS account
-			check(LocalUserNumToLastLoginCredentials.Contains(LocalUserNum));
-			FConnectCredentials Credentials(ToEOS_EExternalCredentialType(GetPlatformOSS()->GetSubsystemName(), *LocalUserNumToLastLoginCredentials[LocalUserNum]), AuthToken);
-			EOS_Connect_LoginOptions Options = { };
-			Options.ApiVersion = EOS_CONNECT_LOGIN_API_LATEST;
-			Options.Credentials = &Credentials;
+				// Now login into our EOS account
+				check(LocalUserNumToLastLoginCredentials.Contains(LocalUserNum));
+				FConnectCredentials Credentials(ToEOS_EExternalCredentialType(GetPlatformOSS()->GetSubsystemName(), *LocalUserNumToLastLoginCredentials[LocalUserNum]), AuthToken);
+				EOS_Connect_LoginOptions Options = { };
+				Options.ApiVersion = EOS_CONNECT_LOGIN_API_LATEST;
+				Options.Credentials = &Credentials;
 
 #if ADD_USER_LOGIN_INFO
-			EOS_Connect_UserLoginInfo UserLoginInfo = {};
-			UserLoginInfo.ApiVersion = EOS_CONNECT_USERLOGININFO_API_LATEST;
-			const FTCHARToUTF8 DisplayNameUtf8(*GetPlatformDisplayName(LocalUserNum));
-			UserLoginInfo.DisplayName = DisplayNameUtf8.Get();
+				EOS_Connect_UserLoginInfo UserLoginInfo = {};
+				UserLoginInfo.ApiVersion = EOS_CONNECT_USERLOGININFO_API_LATEST;
+				const FTCHARToUTF8 DisplayNameUtf8(*GetPlatformDisplayName(LocalUserNum));
+				UserLoginInfo.DisplayName = DisplayNameUtf8.Get();
 
-			Options.UserLoginInfo = &UserLoginInfo;
+				Options.UserLoginInfo = &UserLoginInfo;
 #endif
 
-			FConnectLoginCallback* CallbackObj = new FConnectLoginCallback(FUserManagerEOSWeakPtr(AsShared()));
-			CallbackObj->CallbackLambda = [this, LocalUserNum](const EOS_Connect_LoginCallbackInfo* Data)
-			{
-				if (Data->ResultCode == EOS_EResult::EOS_Success)
-				{
-					// We have an account mapping to the platform account, skip to final login
-					FullLoginCallback(LocalUserNum, nullptr, Data->LocalUserId);
-				}
-				else if (Data->ResultCode == EOS_EResult::EOS_InvalidUser)
-				{
-					// We need to create the platform account mapping for this user using the continuation token
-					CreateConnectedLogin(LocalUserNum, nullptr, Data->ContinuanceToken);
-				}
-				else
+				FConnectLoginCallback* CallbackObj = new FConnectLoginCallback(AsWeak());
+				CallbackObj->CallbackLambda = [this, LocalUserNum](const EOS_Connect_LoginCallbackInfo* Data)
 				{
-					const FString ErrorString = FString::Printf(TEXT("ConnectLoginNoEAS(%d) failed with EOS result code (%s)"), LocalUserNum, ANSI_TO_TCHAR(EOS_EResult_ToString(Data->ResultCode)));
-					UE_LOG_ONLINE(Warning, TEXT("%s"), *ErrorString);
-					TriggerOnLoginCompleteDelegates(LocalUserNum, false, *FUniqueNetIdEOS::EmptyId(), ErrorString);
-				}
-			};
-			EOS_Connect_Login(EOSSubsystem->ConnectHandle, &Options, CallbackObj, CallbackObj->GetCallbackPtr());
+					if (Data->ResultCode == EOS_EResult::EOS_Success)
+					{
+						// We have an account mapping to the platform account, skip to final login
+						FullLoginCallback(LocalUserNum, nullptr, Data->LocalUserId);
+					}
+					else if (Data->ResultCode == EOS_EResult::EOS_InvalidUser)
+					{
+						// We need to create the platform account mapping for this user using the continuation token
+						CreateConnectedLogin(LocalUserNum, nullptr, Data->ContinuanceToken);
+					}
+					else
+					{
+						const FString ErrorString = FString::Printf(TEXT("ConnectLoginNoEAS(%d) failed with EOS result code (%s)"), LocalUserNum, ANSI_TO_TCHAR(EOS_EResult_ToString(Data->ResultCode)));
+						UE_LOG_ONLINE(Warning, TEXT("%s"), *ErrorString);
+						TriggerOnLoginCompleteDelegates(LocalUserNum, false, *FUniqueNetIdEOS::EmptyId(), ErrorString);
+					}
+				};
+				EOS_Connect_Login(EOSSubsystem->ConnectHandle, &Options, CallbackObj, CallbackObj->GetCallbackPtr());
+			}
 		}));
 
 	return true;
@@ -718,7 +724,7 @@ bool FUserManagerEOS::ConnectLoginEAS(int32 LocalUserNum, EOS_EpicAccountId Acco
 		Options.ApiVersion = EOS_CONNECT_LOGIN_API_LATEST;
 		Options.Credentials = &Credentials;
 
-		FConnectLoginCallback* CallbackObj = new FConnectLoginCallback(FUserManagerEOSWeakPtr(AsShared()));
+		FConnectLoginCallback* CallbackObj = new FConnectLoginCallback(AsWeak());
 		CallbackObj->CallbackLambda = [LocalUserNum, AccountId, this](const EOS_Connect_LoginCallbackInfo* Data)
 		{
 			if (Data->ResultCode == EOS_EResult::EOS_Success)
@@ -783,7 +789,7 @@ void FUserManagerEOS::RefreshConnectLogin(int32 LocalUserNum)
 			Options.ApiVersion = EOS_CONNECT_LOGIN_API_LATEST;
 			Options.Credentials = &Credentials;
 
-			FConnectLoginCallback* CallbackObj = new FConnectLoginCallback(FUserManagerEOSWeakPtr(AsShared()));
+			FConnectLoginCallback* CallbackObj = new FConnectLoginCallback(AsWeak());
 			CallbackObj->CallbackLambda = [LocalUserNum, AccountId, this](const EOS_Connect_LoginCallbackInfo* Data)
 			{
 				if (Data->ResultCode != EOS_EResult::EOS_Success)
@@ -806,34 +812,37 @@ void FUserManagerEOS::RefreshConnectLogin(int32 LocalUserNum)
 	{
 		// Not using EAS so grab the platform auth token
 		GetPlatformAuthToken(LocalUserNum,
-			FOnGetLinkedAccountAuthTokenCompleteDelegate::CreateLambda([this](int32 LocalUserNum, bool bWasSuccessful, const FExternalAuthToken& AuthToken)
+			FOnGetLinkedAccountAuthTokenCompleteDelegate::CreateLambda([this, WeakThis = AsWeak()](int32 LocalUserNum, bool bWasSuccessful, const FExternalAuthToken& AuthToken)
 			{
-				if (!bWasSuccessful || !AuthToken.IsValid())
+				if (FUserManagerEOSPtr StrongThis = WeakThis.Pin())
 				{
-					UE_LOG_ONLINE(Error, TEXT("ConnectLoginNoEAS(%d) failed due to the platform OSS giving an empty auth token"), LocalUserNum);
-					Logout(LocalUserNum);
-					return;
-				}
-
-				// Now login into our EOS account
-				check(LocalUserNumToLastLoginCredentials.Contains(LocalUserNum));
-				const FOnlineAccountCredentials& Creds = *LocalUserNumToLastLoginCredentials[LocalUserNum];
-				EOS_EExternalCredentialType CredType = ToEOS_EExternalCredentialType(GetPlatformOSS()->GetSubsystemName(), Creds);
-				FConnectCredentials Credentials(CredType, AuthToken);
-				EOS_Connect_LoginOptions Options = { };
-				Options.ApiVersion = EOS_CONNECT_LOGIN_API_LATEST;
-				Options.Credentials = &Credentials;
-
-				FConnectLoginCallback* CallbackObj = new FConnectLoginCallback(FUserManagerEOSWeakPtr(AsShared()));
-				CallbackObj->CallbackLambda = [this, LocalUserNum](const EOS_Connect_LoginCallbackInfo* Data)
-				{
-					if (Data->ResultCode != EOS_EResult::EOS_Success)
+					if (!bWasSuccessful || !AuthToken.IsValid())
 					{
-						UE_LOG_ONLINE(Error, TEXT("Failed to refresh ConnectLogin(%d) failed with EOS result code (%s)"), LocalUserNum, ANSI_TO_TCHAR(EOS_EResult_ToString(Data->ResultCode)));
+						UE_LOG_ONLINE(Error, TEXT("ConnectLoginNoEAS(%d) failed due to the platform OSS giving an empty auth token"), LocalUserNum);
 						Logout(LocalUserNum);
+						return;
 					}
-				};
-				EOS_Connect_Login(EOSSubsystem->ConnectHandle, &Options, CallbackObj, CallbackObj->GetCallbackPtr());
+
+					// Now login into our EOS account
+					check(LocalUserNumToLastLoginCredentials.Contains(LocalUserNum));
+					const FOnlineAccountCredentials& Creds = *LocalUserNumToLastLoginCredentials[LocalUserNum];
+					EOS_EExternalCredentialType CredType = ToEOS_EExternalCredentialType(GetPlatformOSS()->GetSubsystemName(), Creds);
+					FConnectCredentials Credentials(CredType, AuthToken);
+					EOS_Connect_LoginOptions Options = { };
+					Options.ApiVersion = EOS_CONNECT_LOGIN_API_LATEST;
+					Options.Credentials = &Credentials;
+
+					FConnectLoginCallback* CallbackObj = new FConnectLoginCallback(AsWeak());
+					CallbackObj->CallbackLambda = [this, LocalUserNum](const EOS_Connect_LoginCallbackInfo* Data)
+					{
+						if (Data->ResultCode != EOS_EResult::EOS_Success)
+						{
+							UE_LOG_ONLINE(Error, TEXT("Failed to refresh ConnectLogin(%d) failed with EOS result code (%s)"), LocalUserNum, ANSI_TO_TCHAR(EOS_EResult_ToString(Data->ResultCode)));
+							Logout(LocalUserNum);
+						}
+					};
+					EOS_Connect_Login(EOSSubsystem->ConnectHandle, &Options, CallbackObj, CallbackObj->GetCallbackPtr());
+				}
 			}));
 	}
 }
@@ -846,7 +855,7 @@ void FUserManagerEOS::CreateConnectedLogin(int32 LocalUserNum, EOS_EpicAccountId
 	Options.ApiVersion = EOS_CONNECT_CREATEUSER_API_LATEST;
 	Options.ContinuanceToken = Token;
 
-	FCreateUserCallback* CallbackObj = new FCreateUserCallback(FUserManagerEOSWeakPtr(AsShared()));
+	FCreateUserCallback* CallbackObj = new FCreateUserCallback(AsWeak());
 	CallbackObj->CallbackLambda = [LocalUserNum, AccountId, this](const EOS_Connect_CreateUserCallbackInfo* Data)
 	{
 		if (Data->ResultCode == EOS_EResult::EOS_Success)
@@ -873,7 +882,7 @@ void FUserManagerEOS::FullLoginCallback(int32 LocalUserNum, EOS_EpicAccountId Ac
 	// Add our login status changed callback if not already set
 	if (LoginNotificationId == 0)
 	{
-		FLoginStatusChangedCallback* CallbackObj = new FLoginStatusChangedCallback(FUserManagerEOSWeakPtr(AsShared()));
+		FLoginStatusChangedCallback* CallbackObj = new FLoginStatusChangedCallback(AsWeak());
 		LoginNotificationCallback = CallbackObj;
 		CallbackObj->CallbackLambda = [this](const EOS_Auth_LoginStatusChangedCallbackInfo* Data)
 		{
@@ -887,7 +896,7 @@ void FUserManagerEOS::FullLoginCallback(int32 LocalUserNum, EOS_EpicAccountId Ac
 	// Register for friends updates if not set yet
 	if (FriendsNotificationId == 0)
 	{
-		FFriendsStatusUpdateCallback* CallbackObj = new FFriendsStatusUpdateCallback(FUserManagerEOSWeakPtr(AsShared()));
+		FFriendsStatusUpdateCallback* CallbackObj = new FFriendsStatusUpdateCallback(AsWeak());
 		FriendsNotificationCallback = CallbackObj;
 		CallbackObj->CallbackLambda = [LocalUserNum, this](const EOS_Friends_OnFriendsUpdateInfo* Data)
 		{
@@ -901,7 +910,7 @@ void FUserManagerEOS::FullLoginCallback(int32 LocalUserNum, EOS_EpicAccountId Ac
 	// Register for presence updates if not set yet
 	if (PresenceNotificationId == 0)
 	{
-		FPresenceChangedCallback* CallbackObj = new FPresenceChangedCallback(FUserManagerEOSWeakPtr(AsShared()));
+		FPresenceChangedCallback* CallbackObj = new FPresenceChangedCallback(AsWeak());
 		PresenceNotificationCallback = CallbackObj;
 		CallbackObj->CallbackLambda = [LocalUserNum, this](const EOS_Presence_PresenceChangedCallbackInfo* Data)
 		{
@@ -923,7 +932,7 @@ void FUserManagerEOS::FullLoginCallback(int32 LocalUserNum, EOS_EpicAccountId Ac
 		FNotificationIdCallbackPair* NotificationPair = new FNotificationIdCallbackPair();
 		LocalUserNumToConnectLoginNotifcationMap.Emplace(LocalUserNum, NotificationPair);
 
-		FRefreshAuthCallback* CallbackObj = new FRefreshAuthCallback(FUserManagerEOSWeakPtr(AsShared()));
+		FRefreshAuthCallback* CallbackObj = new FRefreshAuthCallback(AsWeak());
 		NotificationPair->Callback = CallbackObj;
 		CallbackObj->CallbackLambda = [LocalUserNum, this](const EOS_Connect_AuthExpirationCallbackInfo* Data)
 		{
@@ -956,10 +965,10 @@ bool FUserManagerEOS::Logout(int32 LocalUserNum)
 		return false;
 	}
 
-	FLogoutCallback* CallbackObj = new FLogoutCallback(FUserManagerEOSWeakPtr(AsShared()));
+	FLogoutCallback* CallbackObj = new FLogoutCallback(AsWeak());
 	CallbackObj->CallbackLambda = [LocalUserNum, this](const EOS_Auth_LogoutCallbackInfo* Data)
 	{
-		FDeletePersistentAuthCallback* DeleteAuthCallbackObj = new FDeletePersistentAuthCallback(FUserManagerEOSWeakPtr(AsShared()));
+		FDeletePersistentAuthCallback* DeleteAuthCallbackObj = new FDeletePersistentAuthCallback(AsWeak());
 		DeleteAuthCallbackObj->CallbackLambda = [this, LocalUserNum, LogoutResultCode = Data->ResultCode](const EOS_Auth_DeletePersistentAuthCallbackInfo* Data)
 		{
 			if (LogoutResultCode == EOS_EResult::EOS_Success)
@@ -1562,7 +1571,7 @@ bool FUserManagerEOS::ShowAccountCreationUI(const int ControllerIndex, const FOn
 {
 	UE_LOG_ONLINE_EXTERNALUI(Warning, TEXT("[FUserManagerEOS::ShowAccountCreationUI] This method is not implemented."));
 	
-	EOSSubsystem->ExecuteNextTick([this, ControllerIndex, Delegate]()
+	EOSSubsystem->ExecuteNextTick([ControllerIndex, Delegate]()
 		{
 			Delegate.ExecuteIfBound(ControllerIndex, FOnlineAccountCredentials(), FOnlineError(EOnlineErrorResult::NotImplemented));
 		});
@@ -1578,8 +1587,8 @@ bool FUserManagerEOS::ShowFriendsUI(int32 LocalUserNum)
 	Options.ApiVersion = EOS_UI_SHOWFRIENDS_API_LATEST;
 	Options.LocalUserId = GetLocalEpicAccountId(LocalUserNum);
 
-	FOnShowFriendsCallback* CallbackObj = new FOnShowFriendsCallback(FUserManagerEOSWeakPtr(AsShared()));
-	CallbackObj->CallbackLambda = [this](const EOS_UI_ShowFriendsCallbackInfo* Data)
+	FOnShowFriendsCallback* CallbackObj = new FOnShowFriendsCallback(AsWeak());
+	CallbackObj->CallbackLambda = [](const EOS_UI_ShowFriendsCallbackInfo* Data)
 	{
 		if (Data->ResultCode == EOS_EResult::EOS_Success)
 		{
@@ -1624,7 +1633,7 @@ bool FUserManagerEOS::ShowWebURL(const FString& Url, const FShowWebUrlParams& Sh
 
 	if (Delegate.IsBound())
 	{
-		EOSSubsystem->ExecuteNextTick([this, Delegate]()
+		EOSSubsystem->ExecuteNextTick([Delegate]()
 			{
 				Delegate.ExecuteIfBound(FString());
 			});
@@ -1646,7 +1655,7 @@ bool FUserManagerEOS::ShowProfileUI(const FUniqueNetId& Requestor, const FUnique
 
 	if (Delegate.IsBound())
 	{
-		EOSSubsystem->ExecuteNextTick([this, Delegate]()
+		EOSSubsystem->ExecuteNextTick([Delegate]()
 			{
 				Delegate.ExecuteIfBound();
 			});
@@ -1668,7 +1677,7 @@ bool FUserManagerEOS::ShowStoreUI(int32 LocalUserNum, const FShowStoreParams& Sh
 
 	if (Delegate.IsBound())
 	{
-		EOSSubsystem->ExecuteNextTick([this, Delegate]()
+		EOSSubsystem->ExecuteNextTick([Delegate]()
 			{
 				Delegate.ExecuteIfBound(false);
 			});
@@ -1683,7 +1692,7 @@ bool FUserManagerEOS::ShowSendMessageUI(int32 LocalUserNum, const FShowSendMessa
 
 	if (Delegate.IsBound())
 	{
-		EOSSubsystem->ExecuteNextTick([this, Delegate]()
+		EOSSubsystem->ExecuteNextTick([Delegate]()
 			{
 				Delegate.ExecuteIfBound(false);
 			});
@@ -1869,7 +1878,7 @@ bool FUserManagerEOS::ReadFriendsList(int32 LocalUserNum, const FString& ListNam
 	Options.ApiVersion = EOS_FRIENDS_QUERYFRIENDS_API_LATEST;
 	Options.LocalUserId = UserNumToAccountIdMap[LocalUserNum];
 
-	FReadFriendsCallback* CallbackObj = new FReadFriendsCallback(FUserManagerEOSWeakPtr(AsShared()));
+	FReadFriendsCallback* CallbackObj = new FReadFriendsCallback(AsWeak());
 	CallbackObj->CallbackLambda = [this, LocalUserNum, ListName, Delegate](const EOS_Friends_QueryFriendsCallbackInfo* Data)
 	{
 		EOS_EResult Result = Data->ResultCode;
@@ -1906,16 +1915,19 @@ bool FUserManagerEOS::ReadFriendsList(int32 LocalUserNum, const FString& ListNam
 				}
 			}
 
-			const TFunction<FOnQueryExternalIdMappingsComplete::TFuncType>& OnExternalIdMappingsQueriedLambda = 
-				[this, LocalUserNum](bool bWasSuccessful, const FUniqueNetId& UserId, const FExternalIdQueryOptions& QueryOptions,
-					const TArray<FString>& ExternalIds, const FString& Error) {
-						ProcessReadFriendsListComplete(LocalUserNum, bWasSuccessful, Error);
-				};
+			const TFunction<FOnQueryExternalIdMappingsComplete::TFuncType>& OnExternalIdMappingsQueriedLambda =
+				[this, WeakThis = AsWeak(), LocalUserNum](bool bWasSuccessful, const FUniqueNetId& UserId, const FExternalIdQueryOptions& QueryOptions, const TArray<FString>& ExternalIds, const FString& Error)
+			{
+				if (FUserManagerEOSPtr StrongThis = WeakThis.Pin())
+				{
+					ProcessReadFriendsListComplete(LocalUserNum, bWasSuccessful, Error);
+				}
+			};
 
 			const auto& ExternalMappingsCallback = OSSInternalCallback::Create<FOnQueryExternalIdMappingsComplete>(EOSSubsystem->UserManager, OnExternalIdMappingsQueriedLambda);
-			
+
 			QueryExternalIdMappings(*GetLocalUniqueNetIdEOS(DefaultLocalUser), FExternalIdQueryOptions(), FriendEasIds, ExternalMappingsCallback);
-			
+
 		}
 		else
 		{
@@ -1954,7 +1966,7 @@ void FUserManagerEOS::SetFriendAlias(int32 LocalUserNum, const FUniqueNetId& Fri
 {
 	UE_LOG_ONLINE_FRIEND(Warning, TEXT("[FUserManagerEOS::SetFriendAlias] This method is not supported."));
 
-	EOSSubsystem->ExecuteNextTick([this, LocalUserNum, FriendId = FriendId.AsShared(), ListName, Delegate]()
+	EOSSubsystem->ExecuteNextTick([LocalUserNum, FriendId = FriendId.AsShared(), ListName, Delegate]()
 		{
 			Delegate.ExecuteIfBound(LocalUserNum, *FriendId, ListName, FOnlineError(EOnlineErrorResult::NotImplemented));
 		});
@@ -1964,7 +1976,7 @@ void FUserManagerEOS::DeleteFriendAlias(int32 LocalUserNum, const FUniqueNetId&
 {
 	UE_LOG_ONLINE_FRIEND(Warning, TEXT("[FUserManagerEOS::DeleteFriendAlias] This method is not supported."));
 
-	EOSSubsystem->ExecuteNextTick([this, LocalUserNum, FriendId = FriendId.AsShared(), ListName, Delegate]()
+	EOSSubsystem->ExecuteNextTick([LocalUserNum, FriendId = FriendId.AsShared(), ListName, Delegate]()
 		{
 			Delegate.ExecuteIfBound(LocalUserNum, *FriendId, ListName, FOnlineError(EOnlineErrorResult::NotImplemented));
 		});
@@ -1974,7 +1986,7 @@ bool FUserManagerEOS::DeleteFriendsList(int32 LocalUserNum, const FString& ListN
 {
 	UE_LOG_ONLINE_FRIEND(Warning, TEXT("[FUserManagerEOS::DeleteFriendsList] This method is not supported."));
 
-	EOSSubsystem->ExecuteNextTick([this, LocalUserNum, ListName, Delegate]()
+	EOSSubsystem->ExecuteNextTick([LocalUserNum, ListName, Delegate]()
 		{
 			Delegate.ExecuteIfBound(LocalUserNum, false, ListName, TEXT("This method is not supported."));
 		});
@@ -2002,7 +2014,7 @@ bool FUserManagerEOS::SendInvite(int32 LocalUserNum, const FUniqueNetId& FriendI
 		return false;
 	}
 
-	FSendInviteCallback* CallbackObj = new FSendInviteCallback(FUserManagerEOSWeakPtr(AsShared()));
+	FSendInviteCallback* CallbackObj = new FSendInviteCallback(AsWeak());
 	CallbackObj->CallbackLambda = [LocalUserNum, ListName, this, Delegate](const EOS_Friends_SendInviteCallbackInfo* Data)
 	{
 		const FString& NetId = AccountIdToStringMap[Data->TargetUserId];
@@ -2045,7 +2057,7 @@ bool FUserManagerEOS::AcceptInvite(int32 LocalUserNum, const FUniqueNetId& Frien
 		return false;
 	}
 
-	FAcceptInviteCallback* CallbackObj = new FAcceptInviteCallback(FUserManagerEOSWeakPtr(AsShared()));
+	FAcceptInviteCallback* CallbackObj = new FAcceptInviteCallback(AsWeak());
 	CallbackObj->CallbackLambda = [LocalUserNum, ListName, this, Delegate](const EOS_Friends_AcceptInviteCallbackInfo* Data)
 	{
 		const FString& NetId = AccountIdToStringMap[Data->TargetUserId];
@@ -2100,9 +2112,12 @@ bool FUserManagerEOS::DeleteFriend(int32 LocalUserNum, const FUniqueNetId& Frien
 {
 	UE_LOG_ONLINE_FRIEND(Warning, TEXT("[FUserManagerEOS::DeleteFriend] Friends may only be deleted via the Epic Games Launcher."));
 
-	EOSSubsystem->ExecuteNextTick([this, LocalUserNum, FriendId = FriendId.AsShared(), ListName]()
+	EOSSubsystem->ExecuteNextTick([this, WeakThis = AsWeak(), LocalUserNum, FriendId = FriendId.AsShared(), ListName]()
 		{
-			TriggerOnDeleteFriendCompleteDelegates(LocalUserNum, false, *FriendId, ListName, TEXT("[FUserManagerEOS::DeleteFriend] Friends may only be deleted via the Epic Games Launcher."));
+			if (FUserManagerEOSPtr StrongThis = WeakThis.Pin())
+			{
+				TriggerOnDeleteFriendCompleteDelegates(LocalUserNum, false, *FriendId, ListName, TEXT("[FUserManagerEOS::DeleteFriend] Friends may only be deleted via the Epic Games Launcher."));
+			}
 		});
 
 	return true;
@@ -2198,9 +2213,12 @@ bool FUserManagerEOS::QueryRecentPlayers(const FUniqueNetId& UserId, const FStri
 {
 	UE_LOG_ONLINE_FRIEND(Warning, TEXT("[FUserManagerEOS::QueryRecentPlayers] This method is not supported."));
 
-	EOSSubsystem->ExecuteNextTick([this, UserId = UserId.AsShared(), Namespace]()
+	EOSSubsystem->ExecuteNextTick([this, WeakThis = AsWeak(), UserId = UserId.AsShared(), Namespace]()
 		{
-			TriggerOnQueryRecentPlayersCompleteDelegates(*UserId, Namespace, false, TEXT("This method is not supported."));
+			if (FUserManagerEOSPtr StrongThis = WeakThis.Pin())
+			{
+				TriggerOnQueryRecentPlayersCompleteDelegates(*UserId, Namespace, false, TEXT("This method is not supported."));
+			}
 		});
 
 	return true;
@@ -2217,9 +2235,12 @@ bool FUserManagerEOS::BlockPlayer(int32 LocalUserNum, const FUniqueNetId& Player
 {
 	UE_LOG_ONLINE_FRIEND(Warning, TEXT("[FUserManagerEOS::BlockPlayer] This method is not supported."));
 
-	EOSSubsystem->ExecuteNextTick([this, LocalUserNum, PlayerId = PlayerId.AsShared()]()
+	EOSSubsystem->ExecuteNextTick([this, WeakThis = AsWeak(), LocalUserNum, PlayerId = PlayerId.AsShared()]()
 		{
-			TriggerOnBlockedPlayerCompleteDelegates(LocalUserNum, false, *PlayerId, TEXT(""), TEXT("This method is not supported"));
+			if (FUserManagerEOSPtr StrongThis = WeakThis.Pin())
+			{
+				TriggerOnBlockedPlayerCompleteDelegates(LocalUserNum, false, *PlayerId, TEXT(""), TEXT("This method is not supported"));
+			}
 		});
 
 	return true;
@@ -2229,9 +2250,12 @@ bool FUserManagerEOS::UnblockPlayer(int32 LocalUserNum, const FUniqueNetId& Play
 {
 	UE_LOG_ONLINE_FRIEND(Warning, TEXT("[FUserManagerEOS::UnblockPlayer] This method is not supported."));
 
-	EOSSubsystem->ExecuteNextTick([this, LocalUserNum, PlayerId = PlayerId.AsShared()]()
+	EOSSubsystem->ExecuteNextTick([this, WeakThis = AsWeak(), LocalUserNum, PlayerId = PlayerId.AsShared()]()
 		{
-			TriggerOnUnblockedPlayerCompleteDelegates(LocalUserNum, false, *PlayerId, TEXT(""), TEXT("This method is not supported"));
+			if (FUserManagerEOSPtr StrongThis = WeakThis.Pin())
+			{
+				TriggerOnUnblockedPlayerCompleteDelegates(LocalUserNum, false, *PlayerId, TEXT(""), TEXT("This method is not supported"));
+			}
 		});
 
 	return true;
@@ -2241,9 +2265,12 @@ bool FUserManagerEOS::QueryBlockedPlayers(const FUniqueNetId& UserId)
 {
 	UE_LOG_ONLINE_FRIEND(Warning, TEXT("[FUserManagerEOS::QueryBlockedPlayers] This method is not supported."));
 
-	EOSSubsystem->ExecuteNextTick([this, UserId = UserId.AsShared()]()
+	EOSSubsystem->ExecuteNextTick([this, WeakThis = AsWeak(), UserId = UserId.AsShared()]()
 		{
-			TriggerOnQueryBlockedPlayersCompleteDelegates(*UserId, false, TEXT("This method is not supported"));
+			if (FUserManagerEOSPtr StrongThis = WeakThis.Pin())
+			{
+				TriggerOnQueryBlockedPlayersCompleteDelegates(*UserId, false, TEXT("This method is not supported"));
+			}
 		});
 
 	return true;
@@ -2444,7 +2471,7 @@ void FUserManagerEOS::SetPresence(const FUniqueNetId& UserId, const FOnlineUserP
 		UE_LOG_ONLINE(Error, TEXT("EOS_PresenceModification_SetData() failed with result code (%s)"), *LexToString(SetDataResult));
 	}
 
-	FSetPresenceCallback* CallbackObj = new FSetPresenceCallback(FUserManagerEOSWeakPtr(AsShared()));
+	FSetPresenceCallback* CallbackObj = new FSetPresenceCallback(AsWeak());
 	CallbackObj->CallbackLambda = [this, Delegate](const EOS_Presence_SetPresenceCallbackInfo* Data)
 	{
 		if (Data->ResultCode == EOS_EResult::EOS_Success && AccountIdToStringMap.Contains(Data->LocalUserId))
@@ -2493,7 +2520,7 @@ void FUserManagerEOS::QueryPresence(const FUniqueNetId& UserId, const FOnPresenc
 	EOS_Bool bHasPresence = EOS_Presence_HasPresence(EOSSubsystem->PresenceHandle, &HasOptions);
 	if (bHasPresence == EOS_FALSE)
 	{
-		FQueryPresenceCallback* CallbackObj = new FQueryPresenceCallback(FUserManagerEOSWeakPtr(AsShared()));
+		FQueryPresenceCallback* CallbackObj = new FQueryPresenceCallback(AsWeak());
 		CallbackObj->CallbackLambda = [this, Delegate](const EOS_Presence_QueryPresenceCallbackInfo* Data)
 		{
 			if (Data->ResultCode == EOS_EResult::EOS_Success && EpicAccountIdToOnlineUserMap.Contains(Data->TargetUserId))
@@ -2647,7 +2674,7 @@ typedef TEOSCallback<EOS_UserInfo_OnQueryUserInfoCallback, EOS_UserInfo_QueryUse
 
 void FUserManagerEOS::ReadUserInfo(int32 LocalUserNum, EOS_EpicAccountId EpicAccountId)
 {
-	FReadUserInfoCallback* CallbackObj = new FReadUserInfoCallback(FUserManagerEOSWeakPtr(AsShared()));
+	FReadUserInfoCallback* CallbackObj = new FReadUserInfoCallback(AsWeak());
 	CallbackObj->CallbackLambda = [this, LocalUserNum, EpicAccountId](const EOS_UserInfo_QueryUserInfoCallbackInfo* Data)
 	{
 		if (Data->ResultCode == EOS_EResult::EOS_Success)
@@ -2728,7 +2755,7 @@ bool FUserManagerEOS::QueryUserIdMapping(const FUniqueNetId& UserId, const FStri
 	}
 	int32 LocalUserNum = GetLocalUserNumFromUniqueNetId(UserId);
 
-	FQueryInfoByNameCallback* CallbackObj = new FQueryInfoByNameCallback(FUserManagerEOSWeakPtr(AsShared()));
+	FQueryInfoByNameCallback* CallbackObj = new FQueryInfoByNameCallback(AsWeak());
 	CallbackObj->CallbackLambda = [LocalUserNum, DisplayNameOrEmail, this, Delegate](const EOS_UserInfo_QueryUserInfoByDisplayNameCallbackInfo* Data)
 	{
 		EOS_EResult Result = Data->ResultCode;
@@ -2847,7 +2874,7 @@ bool FUserManagerEOS::QueryExternalIdMappings(const FUniqueNetId& UserId, const
 			FCStringAnsi::Strncpy(Options.PointerArray[ProcessedCount], TCHAR_TO_UTF8(*ExternalIds[ProcessedCount]), EOS_CONNECT_EXTERNAL_ACCOUNT_ID_MAX_LENGTH+1);
 			BatchIds.Add(ExternalIds[ProcessedCount]);
 		}
-		FQueryByStringIdsCallback* CallbackObj = new FQueryByStringIdsCallback(FUserManagerEOSWeakPtr(AsShared()));
+		FQueryByStringIdsCallback* CallbackObj = new FQueryByStringIdsCallback(AsWeak());
 		CallbackObj->CallbackLambda = [LocalUserNum, QueryOptions, BatchIds, this, Delegate](const EOS_Connect_QueryExternalAccountMappingsCallbackInfo* Data)
 		{
 			EOS_EResult Result = Data->ResultCode;
@@ -2868,7 +2895,7 @@ bool FUserManagerEOS::QueryExternalIdMappings(const FUniqueNetId& UserId, const
 				// Get the product id for each epic account passed in
 				for (const FString& StringId : BatchIds)
 				{
-					FCStringAnsi::Strncpy(Options.AccountId, TCHAR_TO_UTF8(*StringId), EOS_CONNECT_EXTERNAL_ACCOUNT_ID_MAX_LENGTH+1);
+					FCStringAnsi::Strncpy(Options.AccountId, TCHAR_TO_UTF8(*StringId), EOS_CONNECT_EXTERNAL_ACCOUNT_ID_MAX_LENGTH + 1);
 					EOS_ProductUserId ProductUserId = EOS_Connect_GetExternalAccountMapping(EOSSubsystem->ConnectHandle, &Options);
 					if (EOS_ProductUserId_IsValid(ProductUserId) == EOS_TRUE)
 					{
@@ -2888,7 +2915,7 @@ bool FUserManagerEOS::QueryExternalIdMappings(const FUniqueNetId& UserId, const
 			{
 				OngoingQueries.RemoveSwap(StringId, false);
 			}
-			
+
 			const bool bWasSuccessful = Result == EOS_EResult::EOS_Success;
 			Delegate.ExecuteIfBound(bWasSuccessful, *EOSID, QueryOptions, BatchIds, ErrorString);
 		};
