Added methods to retrieve online session from LobbyId
Fix for crash in CopyLobbyData when the session data was outdated

diff --git a/Engine/Plugins/Online/OnlineSubsystemEOS/Source/OnlineSubsystemEOS/Private/OnlineSessionEOS.cpp b/Engine/Plugins/Online/OnlineSubsystemEOS/Source/OnlineSubsystemEOS/Private/OnlineSessionEOS.cpp
--- a/Engine/Plugins/Online/OnlineSubsystemEOS/Source/OnlineSubsystemEOS/Private/OnlineSessionEOS.cpp
+++ b/Engine/Plugins/Online/OnlineSubsystemEOS/Source/OnlineSubsystemEOS/Private/OnlineSessionEOS.cpp
@@ -515,6 +515,8 @@ FNamedOnlineSession* FOnlineSessionEOS::GetNamedSessionFromLobbyId(const FUnique
 		if (Session.SessionInfo.IsValid())
 		{
 			FOnlineSessionInfoEOS* SessionInfo = (FOnlineSessionInfoEOS*)Session.SessionInfo.Get();
+
+			// We'll check if the session is a Lobby session before comparing the ids
 			if (!Session.SessionSettings.bIsLANMatch && Session.SessionSettings.bUseLobbiesIfAvailable && *SessionInfo->SessionId == LobbyId)
 			{
 				Result = &Sessions[SearchIndex];
@@ -526,6 +528,83 @@ FNamedOnlineSession* FOnlineSessionEOS::GetNamedSessionFromLobbyId(const FUnique
 	return Result;
 }
 
+/**
+ * Searches the search results and invites arrays for the specified session
+ *
+ * @param LobbyId the lobby id to search for
+ *
+ * @return pointer to the struct if found, NULL otherwise
+ */
+FOnlineSessionSearchResult* FOnlineSessionEOS::GetSearchResultFromLobbyId(const FUniqueNetIdEOSLobby& LobbyId)
+{
+	FOnlineSessionSearchResult* Result = nullptr;
+
+	TArray<FOnlineSessionSearchResult*> CombinedSearchResults;
+
+	if (CurrentSessionSearch.IsValid())
+	{
+		for (FOnlineSessionSearchResult& SearchResult : CurrentSessionSearch->SearchResults)
+		{
+			CombinedSearchResults.Add(&SearchResult);
+		}
+	}
+
+	if (LastInviteSearch.IsValid())
+	{
+		for (FOnlineSessionSearchResult& SearchResult : LastInviteSearch->SearchResults)
+		{
+			CombinedSearchResults.Add(&SearchResult);
+		}
+	}
+
+	for (FOnlineSessionSearchResult* SearchResult : CombinedSearchResults)
+	{
+		const FOnlineSession& Session = SearchResult->Session;
+		if (Session.SessionInfo.IsValid())
+		{
+			FOnlineSessionInfoEOS* SessionInfo = (FOnlineSessionInfoEOS*)Session.SessionInfo.Get();
+
+			// We'll check if the session is a Lobby session before comparing the ids
+			if (!Session.SessionSettings.bIsLANMatch && Session.SessionSettings.bUseLobbiesIfAvailable && *SessionInfo->SessionId == LobbyId)
+			{
+				Result = SearchResult;
+				break;
+			}
+		}
+	}
+
+	return Result;
+}
+
+/**
+ * Searches all local sessions containers for the specified session
+ *
+ * @param LobbyId the lobby id to search for
+ *
+ * @return pointer to the struct if found, NULL otherwise
+ */
+FOnlineSession* FOnlineSessionEOS::GetOnlineSessionFromLobbyId(const FUniqueNetIdEOSLobby& LobbyId)
+{
+	// First we try to retrieve a named session matching the given lobby id
+
+	FOnlineSession* Result = GetNamedSessionFromLobbyId(LobbyId);
+
+	if (!Result)
+	{
+		// If no named session were found with that lobby id, we look amongst the sessions in the latest search results
+		if (FOnlineSessionSearchResult* SearchResult = GetSearchResultFromLobbyId(LobbyId))
+		{
+			Result = &SearchResult->Session;
+		}
+		else
+		{
+			UE_LOG_ONLINE_SESSION(Verbose, TEXT("[FOnlineSessionEOS::GetOnlineSessionFromLobbyId] Session with LobbyId [%s] not found."), *LobbyId.ToString());
+		}
+	}
+
+	return Result;
+}
+
 void FOnlineSessionEOS::RegisterLobbyNotifications()
 {
 	// Lobby data updates
@@ -4157,20 +4236,30 @@ void FOnlineSessionEOS::CopyLobbyData(const TSharedRef<FLobbyDetailsEOS>& LobbyD
 		TargetUserIds.Add(TargetUserId);
 	}
 
-	EOSSubsystem->UserManager->ResolveUniqueNetIds(TargetUserIds, [this, LobbyDetails, &OutSession, OriginalCallback = Callback](TMap<EOS_ProductUserId, FUniqueNetIdEOSRef> ResolvedUniqueNetIds)
+	if (!TargetUserIds.IsEmpty())
 	{
-		for (TMap<EOS_ProductUserId, FUniqueNetIdEOSRef>::TConstIterator It(ResolvedUniqueNetIds); It; ++It)
-		{
-			if (!OutSession.SessionSettings.MemberSettings.Contains(It.Value()))
+		EOSSubsystem->UserManager->ResolveUniqueNetIds(TargetUserIds, [this, LobbyDetails, LobbyId = FUniqueNetIdEOSLobby::Create(LobbyDetailsInfo->LobbyId), OriginalCallback = Callback](TMap<EOS_ProductUserId, FUniqueNetIdEOSRef> ResolvedUniqueNetIds)
 			{
-				OutSession.SessionSettings.MemberSettings.Add(It.Value(), FSessionSettings());
-			}
+				
+				FOnlineSession* Session = GetOnlineSessionFromLobbyId(*LobbyId);
+				if (Session)
+				{
+					for (TMap<EOS_ProductUserId, FUniqueNetIdEOSRef>::TConstIterator It(ResolvedUniqueNetIds); It; ++It)
+					{
+						FSessionSettings& MemberSettings = Session->SessionSettings.MemberSettings.FindOrAdd(It.Value());
 
-			CopyLobbyMemberAttributes(*LobbyDetails, It.Key(), OutSession.SessionSettings.MemberSettings[It.Value()]);
-		}
+						CopyLobbyMemberAttributes(*LobbyDetails, It.Key(), MemberSettings);
+					}
+				}
 
-		OriginalCallback(true);
-	});
+				const bool bWasSuccessful = Session != nullptr;
+				OriginalCallback(bWasSuccessful);
+			});
+	}
+	else
+	{
+		Callback(true);
+	}
 }
 
 void FOnlineSessionEOS::CopyLobbyAttributes(const TSharedRef<FLobbyDetailsEOS>& LobbyDetails, FOnlineSession& OutSession)
diff --git a/Engine/Plugins/Online/OnlineSubsystemEOS/Source/OnlineSubsystemEOS/Private/OnlineSessionEOS.h b/Engine/Plugins/Online/OnlineSubsystemEOS/Source/OnlineSubsystemEOS/Private/OnlineSessionEOS.h
--- a/Engine/Plugins/Online/OnlineSubsystemEOS/Source/OnlineSubsystemEOS/Private/OnlineSessionEOS.h
+++ b/Engine/Plugins/Online/OnlineSubsystemEOS/Source/OnlineSubsystemEOS/Private/OnlineSessionEOS.h
@@ -269,6 +269,8 @@ private:
 	void GetEpicAccountIdAsync(const EOS_ProductUserId& ProductUserId, const GetEpicAccountIdAsyncCallback& Callback);
 	void RegisterLobbyNotifications();
 	FNamedOnlineSession* GetNamedSessionFromLobbyId(const FUniqueNetIdEOSLobby& LobbyId);
+	FOnlineSessionSearchResult* GetSearchResultFromLobbyId(const FUniqueNetIdEOSLobby& LobbyId);
+	FOnlineSession* GetOnlineSessionFromLobbyId(const FUniqueNetIdEOSLobby& LobbyId);
 	bool GetEpicAccountIdFromProductUserId(const EOS_ProductUserId& ProductUserId, EOS_EpicAccountId& EpicAccountId);
 	EOS_ELobbyPermissionLevel GetLobbyPermissionLevelFromSessionSettings(const FOnlineSessionSettings& SessionSettings);
 	uint32_t GetLobbyMaxMembersFromSessionSettings(const FOnlineSessionSettings& SessionSettings);
