Add support for password logins to OnlineSubsystemEOS.

diff --git a/Engine/Plugins/Online/OnlineSubsystemEOS/Source/OnlineSubsystemEOS/Private/UserManagerEOS.cpp b/Engine/Plugins/Online/OnlineSubsystemEOS/Source/OnlineSubsystemEOS/Private/UserManagerEOS.cpp
--- a/Engine/Plugins/Online/OnlineSubsystemEOS/Source/OnlineSubsystemEOS/Private/UserManagerEOS.cpp
+++ b/Engine/Plugins/Online/OnlineSubsystemEOS/Source/OnlineSubsystemEOS/Private/UserManagerEOS.cpp
@@ -452,6 +452,13 @@ bool FUserManagerEOS::Login(int32 LocalUserNum, const FOnlineAccountCredentials&
 		FCStringAnsi::Strncpy(Credentials.IdAnsi, TCHAR_TO_UTF8(*AccountCredentials.Id), EOS_OSS_STRING_BUFFER_LENGTH);
 		FCStringAnsi::Strncpy(Credentials.TokenAnsi, TCHAR_TO_UTF8(*AccountCredentials.Token), EOS_MAX_TOKEN_SIZE);
 	}
+	else if (AccountCredentials.Type == TEXT("password"))
+	{
+		// This is using a direct username / password. Restricted and not generally available.
+		Credentials.Type = EOS_ELoginCredentialType::EOS_LCT_Password;
+		FCStringAnsi::Strncpy(Credentials.IdAnsi, TCHAR_TO_UTF8(*AccountCredentials.Id), EOS_OSS_STRING_BUFFER_LENGTH);
+		FCStringAnsi::Strncpy(Credentials.TokenAnsi, TCHAR_TO_UTF8(*AccountCredentials.Token), EOS_MAX_TOKEN_SIZE);
+	}
 	else if (AccountCredentials.Type == TEXT("accountportal"))
 	{
 		// This is auth via the EOS Account Portal
