- Fix full price items being reported as on sale, and free items as not being on sale.

diff --git a/Engine/Plugins/Online/OnlineSubsystemEOS/Source/OnlineSubsystemEOS/Private/OnlineStoreEOS.cpp b/Engine/Plugins/Online/OnlineSubsystemEOS/Source/OnlineSubsystemEOS/Private/OnlineStoreEOS.cpp
--- a/Engine/Plugins/Online/OnlineSubsystemEOS/Source/OnlineSubsystemEOS/Private/OnlineStoreEOS.cpp
+++ b/Engine/Plugins/Online/OnlineSubsystemEOS/Source/OnlineSubsystemEOS/Private/OnlineStoreEOS.cpp
@@ -109,7 +109,7 @@ void FOnlineStoreEOS::QueryOffers(const FUniqueNetId& UserId, const FOnQueryOnli
 				OfferRef->RegularPrice = Offer->OriginalPrice;
 				OfferRef->NumericPrice = Offer->CurrentPrice;
 #endif
-				OfferRef->DiscountType = Offer->DiscountPercentage == 0 ? EOnlineStoreOfferDiscountType::NotOnSale : EOnlineStoreOfferDiscountType::DiscountAmount;
+				OfferRef->DiscountType = Offer->DiscountPercentage == 100 ? EOnlineStoreOfferDiscountType::NotOnSale : EOnlineStoreOfferDiscountType::DiscountAmount;
 			}
 
 			OfferRef->RegularPriceText = FText::AsCurrencyBase(OfferRef->RegularPrice, OfferRef->CurrencyCode, NULL, Offer->DecimalPoint);
