Enable passing the EOS_PF_LOADING_IN_EDITOR EOS overlay flag.

diff --git a/Engine/Plugins/Online/OnlineSubsystemEOS/Source/OnlineSubsystemEOS/Private/OnlineSubsystemEOS.cpp b/Engine/Plugins/Online/OnlineSubsystemEOS/Source/OnlineSubsystemEOS/Private/OnlineSubsystemEOS.cpp
--- a/Engine/Plugins/Online/OnlineSubsystemEOS/Source/OnlineSubsystemEOS/Private/OnlineSubsystemEOS.cpp
+++ b/Engine/Plugins/Online/OnlineSubsystemEOS/Source/OnlineSubsystemEOS/Private/OnlineSubsystemEOS.cpp
@@ -224,6 +224,12 @@ bool FOnlineSubsystemEOS::PlatformCreate()
 	{
 		OverlayFlags |= EOS_PF_DISABLE_SOCIAL_OVERLAY;
 	}
+#if WITH_EDITOR
+	if (!EOSSettings.bEnableEditorOverlay)
+	{
+		OverlayFlags |= EOS_PF_LOADING_IN_EDITOR;
+	}
+#endif
 
 	// Don't allow the overlay to be used in the editor when running PIE.
 	const bool bEditorOverlayAllowed = EOSSettings.bEnableEditorOverlay && InstanceName == FOnlineSubsystemImpl::DefaultInstanceName;
