- Fix AutoLogin not being usable with -AUTH_TYPE=accountportal due to missing AUTH_LOGIN and AUTH_PASSWORD params

diff --git a/Engine/Plugins/Online/OnlineSubsystemEOS/Source/OnlineSubsystemEOS/Private/UserManagerEOS.cpp b/Engine/Plugins/Online/OnlineSubsystemEOS/Source/OnlineSubsystemEOS/Private/UserManagerEOS.cpp
--- a/Engine/Plugins/Online/OnlineSubsystemEOS/Source/OnlineSubsystemEOS/Private/UserManagerEOS.cpp
+++ b/Engine/Plugins/Online/OnlineSubsystemEOS/Source/OnlineSubsystemEOS/Private/UserManagerEOS.cpp
@@ -1019,7 +1019,7 @@ bool FUserManagerEOS::AutoLogin(int32 LocalUserNum)
 
 	FEOSSettings Settings = UEOSSettings::GetSettings();
 
-	if (EOSSubsystem->bIsDefaultOSS && Settings.bUseEAS && (LoginId.IsEmpty() || Password.IsEmpty() || AuthType.IsEmpty()))
+	if (EOSSubsystem->bIsDefaultOSS && Settings.bUseEAS && AuthType.IsEmpty())
 	{
 		UE_LOG_ONLINE(Warning, TEXT("Unable to AutoLogin user (%d) due to missing auth command line args"), LocalUserNum);
 		return false;
