Refactor of session participant change delegates.
Deprecated OnSessionParticipantsChanged and OnSessionParticipantRemoved for 5.2 and added new OnSessionParticipantJoined and OnSessionParticipantLeft.
Disabled deprecation warnings in references to old delegates.
Changed callsites for deprecated delegates to use the new ones.

diff --git a/Engine/Plugins/Online/OnlineServicesOSSAdapter/Source/Private/Online/SessionsOSSAdapter.cpp b/Engine/Plugins/Online/OnlineServicesOSSAdapter/Source/Private/Online/SessionsOSSAdapter.cpp
--- a/Engine/Plugins/Online/OnlineServicesOSSAdapter/Source/Private/Online/SessionsOSSAdapter.cpp
+++ b/Engine/Plugins/Online/OnlineServicesOSSAdapter/Source/Private/Online/SessionsOSSAdapter.cpp
@@ -317,7 +317,7 @@ void FSessionsOSSAdapter::Initialize()
 		SessionEvents.OnUISessionJoinRequested.Broadcast(Event);
 	}));
 
-	SessionsInterface->AddOnSessionParticipantRemovedDelegate_Handle(FOnSessionParticipantRemovedDelegate::CreateLambda([this](FName SessionName, const FUniqueNetId& TargetUniqueNetId)
+	SessionsInterface->AddOnSessionParticipantJoinedDelegate_Handle(FOnSessionParticipantJoinedDelegate::CreateLambda([this](FName SessionName, const FUniqueNetId& TargetUniqueNetId)
 	{
 		// We won't update a session that can't be retrieved from the OSS or doesn't exist in OnlineServices anymore
 		if (UpdateV2Session(SessionName))
@@ -327,7 +327,7 @@ void FSessionsOSSAdapter::Initialize()
 			const FAccountId TargetAccountId = ServicesOSSAdapter.GetAccountIdRegistry().FindOrAddHandle(TargetUniqueNetIdRef);
 
 			FSessionUpdate SessionUpdate;
-			SessionUpdate.RemovedSessionMembers.Add(TargetAccountId);
+			SessionUpdate.AddedSessionMembers.Emplace(TargetAccountId);
 
 			const FSessionUpdated Event{ SessionName , SessionUpdate };
 
@@ -335,7 +335,7 @@ void FSessionsOSSAdapter::Initialize()
 		}
 	}));
 
-	SessionsInterface->AddOnSessionParticipantsChangeDelegate_Handle(FOnSessionParticipantsChangeDelegate::CreateLambda([this](FName SessionName, const FUniqueNetId& TargetUniqueNetId, bool bJoined)
+	SessionsInterface->AddOnSessionParticipantLeftDelegate_Handle(FOnSessionParticipantLeftDelegate::CreateLambda([this](FName SessionName, const FUniqueNetId& TargetUniqueNetId, EOnSessionParticipantLeftReason LeaveReason)
 	{
 		// We won't update a session that can't be retrieved from the OSS or doesn't exist in OnlineServices anymore
 		if (UpdateV2Session(SessionName))
@@ -345,14 +345,7 @@ void FSessionsOSSAdapter::Initialize()
 			const FAccountId TargetAccountId = ServicesOSSAdapter.GetAccountIdRegistry().FindOrAddHandle(TargetUniqueNetIdRef);
 
 			FSessionUpdate SessionUpdate;
-			if (bJoined)
-			{
-				SessionUpdate.AddedSessionMembers.Emplace(TargetAccountId);
-			}
-			else
-			{
-				SessionUpdate.RemovedSessionMembers.Emplace(TargetAccountId);
-			}
+			SessionUpdate.RemovedSessionMembers.Emplace(TargetAccountId);
 
 			const FSessionUpdated Event{ SessionName , SessionUpdate };
 
@@ -367,8 +360,8 @@ void FSessionsOSSAdapter::Shutdown()
 {
 	SessionsInterface->ClearOnSessionInviteReceivedDelegates(this);
 	SessionsInterface->ClearOnSessionUserInviteAcceptedDelegates(this);
-	SessionsInterface->ClearOnSessionParticipantRemovedDelegates(this);
-	SessionsInterface->ClearOnSessionParticipantsChangeDelegates(this);
+	SessionsInterface->ClearOnSessionParticipantJoinedDelegates(this);
+	SessionsInterface->ClearOnSessionParticipantLeftDelegates(this);
 	SessionsInterface->ClearOnSessionParticipantSettingsUpdatedDelegates(this);
 
 	Super::Shutdown();
diff --git a/Engine/Plugins/Online/OnlineSubsystem/Source/Private/OnlineSessionInterface.cpp b/Engine/Plugins/Online/OnlineSubsystem/Source/Private/OnlineSessionInterface.cpp
--- a/Engine/Plugins/Online/OnlineSubsystem/Source/Private/OnlineSessionInterface.cpp
+++ b/Engine/Plugins/Online/OnlineSubsystem/Source/Private/OnlineSessionInterface.cpp
@@ -4,6 +4,20 @@
 #include "OnlineSubsystem.h"
 #include "OnlineError.h"
 
+const TCHAR* ToLogString(EOnSessionParticipantLeftReason LeaveReason)
+{
+	switch (LeaveReason)
+	{
+	case EOnSessionParticipantLeftReason::Left:			return TEXT("Left");
+	case EOnSessionParticipantLeftReason::Disconnected:	return TEXT("Disconnected");
+	case EOnSessionParticipantLeftReason::Kicked:		return TEXT("Kicked");
+	case EOnSessionParticipantLeftReason::Closed:		return TEXT("Closed");
+	}
+
+	checkNoEntry();
+	return TEXT("Invalid");
+}
+
 // Workaround for not being able to clear a delegate handle for a lambda while the lambda is being executed
 void OnOnlineSessionStartMatchmakingBroadcast(FName DelegateSessionName, bool bWasSuccessful, IOnlineSession* OnlineSession, FName RequestedSessionName, FDelegateHandle* DelegateHandle, const FOnStartMatchmakingComplete CompletionDelegate)
 {
diff --git a/Engine/Plugins/Online/OnlineSubsystem/Source/Public/Interfaces/OnlineSessionInterface.h b/Engine/Plugins/Online/OnlineSubsystem/Source/Public/Interfaces/OnlineSessionInterface.h
--- a/Engine/Plugins/Online/OnlineSubsystem/Source/Public/Interfaces/OnlineSessionInterface.h
+++ b/Engine/Plugins/Online/OnlineSubsystem/Source/Public/Interfaces/OnlineSessionInterface.h
@@ -172,8 +172,56 @@ typedef FOnJoinSessionComplete::FDelegate FOnJoinSessionCompleteDelegate;
  * @param UniqueId The ID of the user whose join state has changed
  * @param bJoined If true this is a join event, (if false it is a leave event)
  */
+UE_DEPRECATED(5.2, "FOnSessionParticipantsChange is deprecated, please use FOnSessionParticipantJoined and FOnSessionParticipantLeft instead.")
 DECLARE_MULTICAST_DELEGATE_ThreeParams(FOnSessionParticipantsChange, FName, const FUniqueNetId&, bool);
+PRAGMA_DISABLE_DEPRECATION_WARNINGS
 typedef FOnSessionParticipantsChange::FDelegate FOnSessionParticipantsChangeDelegate;
+PRAGMA_ENABLE_DEPRECATION_WARNINGS
+
+/**
+ * Delegate fired when a session's member is removed from the session
+ * @param SessionName The name of the session
+ * @param TargetUniqueNetId The UniqueNetId of the member who was removed
+ */
+UE_DEPRECATED(5.2, "FOnSessionParticipantRemoved is deprecated, please use FOnSessionParticipantLeft instead.")
+DECLARE_MULTICAST_DELEGATE_TwoParams(FOnSessionParticipantRemoved, FName, const FUniqueNetId&);
+PRAGMA_DISABLE_DEPRECATION_WARNINGS
+typedef FOnSessionParticipantRemoved::FDelegate FOnSessionParticipantRemovedDelegate;
+PRAGMA_ENABLE_DEPRECATION_WARNINGS
+
+/**
+ * Delegate fired when a player has joined a session
+ *
+ * @param SessionName The name of the session that changed
+ * @param UniqueId The ID of the user who joined
+ */
+DECLARE_MULTICAST_DELEGATE_TwoParams(FOnSessionParticipantJoined, FName, const FUniqueNetId&);
+typedef FOnSessionParticipantJoined::FDelegate FOnSessionParticipantJoinedDelegate;
+
+/** Possible motives for a participant leaving an online session */
+enum class EOnSessionParticipantLeftReason : uint8
+{
+	/** The participant left the session of their own accord */
+	Left,
+	/** The participant got disconnected from a session that is still active */
+	Disconnected,
+	/** The participant was forcefully removed from the session */
+	Kicked,
+	/** The session ended and the participant got removed from it */
+	Closed
+};
+
+ONLINESUBSYSTEM_API const TCHAR* ToLogString(EOnSessionParticipantLeftReason LeaveReason);
+
+/**
+ * Delegate fired when a player has left a session
+ *
+ * @param SessionName The name of the session that changed
+ * @param UniqueId The ID of the user who left
+ * @param LeaveReason An enum defining the motive for the participant's leave
+ */
+DECLARE_MULTICAST_DELEGATE_ThreeParams(FOnSessionParticipantLeft, FName, const FUniqueNetId&, EOnSessionParticipantLeftReason);
+typedef FOnSessionParticipantLeft::FDelegate FOnSessionParticipantLeftDelegate;
 
 /**
  * Delegate fired when session is requesting QOS measurements
@@ -182,17 +230,6 @@ typedef FOnSessionParticipantsChange::FDelegate FOnSessionParticipantsChangeDele
 DECLARE_MULTICAST_DELEGATE_OneParam(FOnQosDataRequested, FName);
 typedef FOnQosDataRequested::FDelegate FOnQosDataRequestedDelegate;
 
-/**
- * Delegate fired when a sessions custom data has changed
- * @param The name of the session that had custom data changed
- * @param the updated session data
- */
-UE_DEPRECATED(4.27, "FOnSessionCustomDataChanged is deprecated, please use FOnSessionSettingsUpdated and FOnSessionParticipantSettingsUpdated instead.")
-DECLARE_MULTICAST_DELEGATE_TwoParams(FOnSessionCustomDataChanged, FName, const FOnlineSessionSettings&);
-PRAGMA_DISABLE_DEPRECATION_WARNINGS
-typedef FOnSessionCustomDataChanged::FDelegate FOnSessionCustomDataChangedDelegate;
-PRAGMA_ENABLE_DEPRECATION_WARNINGS
-
 /**
  * Delegate fired when a session's settings have changed
  * @param The name of the session that had its settings changed
@@ -210,14 +247,6 @@ typedef FOnSessionSettingsUpdated::FDelegate FOnSessionSettingsUpdatedDelegate;
 DECLARE_MULTICAST_DELEGATE_ThreeParams(FOnSessionParticipantSettingsUpdated, FName, const FUniqueNetId&, const FOnlineSessionSettings&);
 typedef FOnSessionParticipantSettingsUpdated::FDelegate FOnSessionParticipantSettingsUpdatedDelegate;
 
-/**
- * Delegate fired when a session's member is removed from the session
- * @param SessionName The name of the session
- * @param TargetUniqueNetId The UniqueNetId of the member who was removed
- */
-DECLARE_MULTICAST_DELEGATE_TwoParams(FOnSessionParticipantRemoved, FName, const FUniqueNetId&);
-typedef FOnSessionParticipantRemoved::FDelegate FOnSessionParticipantRemovedDelegate;
-
 /**
  * Delegate fired once a single search result is returned (ie friend invite / join)
  * Session has not been joined at this point, and requires a call to JoinSession()
@@ -710,6 +739,7 @@ public:
 	 */
 	DEFINE_ONLINE_DELEGATE_TWO_PARAM(OnJoinSessionComplete, FName, EOnJoinSessionCompleteResult::Type);
 
+	PRAGMA_DISABLE_DEPRECATION_WARNINGS
 	/**
 	 * Delegate fired when a player has joined or left a session
 	 * @param SessionName The name of the session that changed
@@ -718,22 +748,36 @@ public:
 	 */
 	DEFINE_ONLINE_DELEGATE_THREE_PARAM(OnSessionParticipantsChange, FName, const FUniqueNetId&, bool);
 
+	/**
+	* Delegate fired when a session's member is removed from the session
+	*
+	* @param SessionName The name of the session
+	* @param TargetUniqueNetId The UniqueNetId of the member who was removed
+	*/
+	DEFINE_ONLINE_DELEGATE_TWO_PARAM(OnSessionParticipantRemoved, FName, const FUniqueNetId&);
+	PRAGMA_ENABLE_DEPRECATION_WARNINGS
+
+	/**
+	* Delegate fired when a player joins a session
+	* @param SessionName The name of the session
+	* @param UniqueId The ID of the user who joined
+	*/
+	DEFINE_ONLINE_DELEGATE_TWO_PARAM(OnSessionParticipantJoined, FName, const FUniqueNetId&);
+
+	/**
+	* Delegate fired when a player leaves a session
+	* @param SessionName The name of the session
+	* @param UniqueId The ID of the user who left
+	* @param LeaveReason An enum defining the motive for the participant's leave
+	*/
+	DEFINE_ONLINE_DELEGATE_THREE_PARAM(OnSessionParticipantLeft, FName, const FUniqueNetId&, EOnSessionParticipantLeftReason);
+
 	/**
 	 * Delegate fired when session is requesting QOS measurements
 	 * @param The name of the session for which the measurements need to be made
 	 */
 	DEFINE_ONLINE_DELEGATE_ONE_PARAM(OnQosDataRequested, FName);
 
-	PRAGMA_DISABLE_DEPRECATION_WARNINGS
-	/**
-	 * Delegate fired when a sessions custom data has been updated
-	 *
-	 * @param SessionName The session that had custom session data changed
-	 * @param SessionSettings The session settings for the session that changed
-	 */
-	DEFINE_ONLINE_DELEGATE_TWO_PARAM(OnSessionCustomDataChanged, FName, const FOnlineSessionSettings&);
-	PRAGMA_ENABLE_DEPRECATION_WARNINGS
-
 	/**
 	 * Delegate fired when a session's settings have changed
 	 *
@@ -751,14 +795,6 @@ public:
 	 */
 	DEFINE_ONLINE_DELEGATE_THREE_PARAM(OnSessionParticipantSettingsUpdated, FName, const FUniqueNetId&, const FOnlineSessionSettings&);
 
-	/**
-	 * Delegate fired when a session's member is removed from the session
-	 *
-	 * @param SessionName The name of the session
-	 * @param TargetUniqueNetId The UniqueNetId of the member who was removed
-	 */
-	DEFINE_ONLINE_DELEGATE_TWO_PARAM(OnSessionParticipantRemoved, FName, const FUniqueNetId&);
-
 	/**
 	 * Allows the local player to follow a friend into a session
 	 *
diff --git a/Engine/Plugins/Online/OnlineSubsystemEOS/Source/OnlineSubsystemEOS/Private/OnlineSessionEOS.cpp b/Engine/Plugins/Online/OnlineSubsystemEOS/Source/OnlineSubsystemEOS/Private/OnlineSessionEOS.cpp
--- a/Engine/Plugins/Online/OnlineSubsystemEOS/Source/OnlineSubsystemEOS/Private/OnlineSessionEOS.cpp
+++ b/Engine/Plugins/Online/OnlineSubsystemEOS/Source/OnlineSubsystemEOS/Private/OnlineSessionEOS.cpp
@@ -714,17 +714,17 @@ void FOnlineSessionEOS::OnLobbyUpdateReceived(const EOS_LobbyId& LobbyId)
 			}
 			else
 			{
-				UE_LOG_ONLINE(Warning, TEXT("[FOnlineSessionEOS::OnLobbyUpdateReceived] EOS_LobbyDetails_CopyInfo not successful. Finished with EOS_EResult %s"), ANSI_TO_TCHAR(EOS_EResult_ToString(CopyInfoResult)));
+				UE_LOG_ONLINE_SESSION(Warning, TEXT("[FOnlineSessionEOS::OnLobbyUpdateReceived] EOS_LobbyDetails_CopyInfo not successful. Finished with EOS_EResult %s"), ANSI_TO_TCHAR(EOS_EResult_ToString(CopyInfoResult)));
 			}
 		}
 		else
 		{
-			UE_LOG_ONLINE(Warning, TEXT("[FOnlineSessionEOS::OnLobbyUpdateReceived] EOS_Lobby_CopyLobbyDetailsHandle not successful. Finished with EOS_EResult %s"), ANSI_TO_TCHAR(EOS_EResult_ToString(CopyLobbyDetailsResult)));
+			UE_LOG_ONLINE_SESSION(Warning, TEXT("[FOnlineSessionEOS::OnLobbyUpdateReceived] EOS_Lobby_CopyLobbyDetailsHandle not successful. Finished with EOS_EResult %s"), ANSI_TO_TCHAR(EOS_EResult_ToString(CopyLobbyDetailsResult)));
 		}
 	}
 	else
 	{
-		UE_LOG_ONLINE(Warning, TEXT("[FOnlineSessionEOS::OnLobbyUpdateReceived] Unable to retrieve session with LobbyId %s"), *LobbyNetId->ToString());
+		UE_LOG_ONLINE_SESSION(Warning, TEXT("[FOnlineSessionEOS::OnLobbyUpdateReceived] Unable to retrieve session with LobbyId %s"), *LobbyNetId->ToString());
 	}
 }
 
@@ -764,17 +764,32 @@ void FOnlineSessionEOS::OnMemberStatusReceived(const EOS_LobbyId& LobbyId, const
 						{
 							RemoveOnlineSessionMember(Session->SessionName, ResolvedUniqueNetId);
 
-							TriggerOnSessionParticipantsChangeDelegates(Session->SessionName, *ResolvedUniqueNetId, false);
+							TriggerOnSessionParticipantLeftDelegates(Session->SessionName, *ResolvedUniqueNetId, EOnSessionParticipantLeftReason::Left);
 						}
 						else
 						{
-							UE_LOG_ONLINE(Warning, TEXT("[FOnlineSessionEOS::OnMemberStatusReceived] EOS_LMS_LEFT: Unable to retrieve session with LobbyId %s"), *LobbyNetId->ToString());
+							UE_LOG_ONLINE_SESSION(VeryVerbose, TEXT("[FOnlineSessionEOS::OnMemberStatusReceived] EOS_LMS_LEFT: Unable to retrieve session with LobbyId %s"), *LobbyNetId->ToString());
 						}
 					});
 			}
 			break;
 		case EOS_ELobbyMemberStatus::EOS_LMS_DISCONNECTED:
-			// OSS Session will end
+			{
+				EOSSubsystem->UserManager->ResolveUniqueNetId(TargetUserId, [this, LobbyNetId](FUniqueNetIdEOSRef ResolvedUniqueNetId)
+					{
+						FNamedOnlineSession* Session = GetNamedSessionFromLobbyId(*LobbyNetId);
+						if (Session)
+						{
+							RemoveOnlineSessionMember(Session->SessionName, ResolvedUniqueNetId);
+
+							TriggerOnSessionParticipantLeftDelegates(Session->SessionName, *ResolvedUniqueNetId, EOnSessionParticipantLeftReason::Disconnected);
+						}
+						else
+						{
+							UE_LOG_ONLINE_SESSION(VeryVerbose, TEXT("[FOnlineSessionEOS::OnMemberStatusReceived] EOS_LMS_DISCONNECTED: Unable to retrieve session with LobbyId %s"), *LobbyNetId->ToString());
+						}
+					});
+			}
 			break;
 		case EOS_ELobbyMemberStatus::EOS_LMS_KICKED:
 			{
@@ -785,11 +800,11 @@ void FOnlineSessionEOS::OnMemberStatusReceived(const EOS_LobbyId& LobbyId, const
 						{
 							RemoveOnlineSessionMember(Session->SessionName, ResolvedUniqueNetId);
 
-							TriggerOnSessionParticipantRemovedDelegates(Session->SessionName, *ResolvedUniqueNetId);
+							TriggerOnSessionParticipantLeftDelegates(Session->SessionName, *ResolvedUniqueNetId, EOnSessionParticipantLeftReason::Kicked);
 						}
 						else
 						{
-							UE_LOG_ONLINE(Warning, TEXT("[FOnlineSessionEOS::OnMemberStatusReceived] EOS_LMS_KICKED: Unable to retrieve session with LobbyId %s"), *LobbyNetId->ToString());
+							UE_LOG_ONLINE_SESSION(VeryVerbose, TEXT("[FOnlineSessionEOS::OnMemberStatusReceived] EOS_LMS_KICKED: Unable to retrieve session with LobbyId %s"), *LobbyNetId->ToString());
 						}
 					});
 			}
@@ -817,19 +832,34 @@ void FOnlineSessionEOS::OnMemberStatusReceived(const EOS_LobbyId& LobbyId, const
 						}
 						else
 						{
-							UE_LOG_ONLINE(Warning, TEXT("[FOnlineSessionEOS::OnMemberStatusReceived] EOS_LMS_PROMOTED: Unable to retrieve session with LobbyId %s"), *LobbyNetId->ToString());
+							UE_LOG_ONLINE_SESSION(Warning, TEXT("[FOnlineSessionEOS::OnMemberStatusReceived] EOS_LMS_PROMOTED: Unable to retrieve session with LobbyId %s"), *LobbyNetId->ToString());
 						}
 					});
 			}
 			break;
 		case EOS_ELobbyMemberStatus::EOS_LMS_CLOSED:
-			// OSS Session will end
+			{
+				EOSSubsystem->UserManager->ResolveUniqueNetId(TargetUserId, [this, LobbyNetId](FUniqueNetIdEOSRef ResolvedUniqueNetId)
+					{
+						FNamedOnlineSession* Session = GetNamedSessionFromLobbyId(*LobbyNetId);
+						if (Session)
+						{
+							RemoveOnlineSessionMember(Session->SessionName, ResolvedUniqueNetId);
+
+							TriggerOnSessionParticipantLeftDelegates(Session->SessionName, *ResolvedUniqueNetId, EOnSessionParticipantLeftReason::Closed);
+						}
+						else
+						{
+							UE_LOG_ONLINE_SESSION(VeryVerbose, TEXT("[FOnlineSessionEOS::OnMemberStatusReceived] EOS_LMS_CLOSED: Unable to retrieve session with LobbyId %s"), *LobbyNetId->ToString());
+						}
+					});
+			}
 			break;
 		}
 	}
 	else
 	{
-		UE_LOG_ONLINE(Warning, TEXT("[FOnlineSessionEOS::OnMemberStatusReceived] Unable to retrieve session with LobbyId %s"), *LobbyNetId->ToString());
+		UE_LOG_ONLINE_SESSION(Warning, TEXT("[FOnlineSessionEOS::OnMemberStatusReceived] Unable to retrieve session with LobbyId %s"), *LobbyNetId->ToString());
 	}
 }
 
@@ -2898,7 +2928,7 @@ void FOnlineSessionEOS::UpdateOrAddLobbyMember(const FUniqueNetIdEOSLobbyRef& Lo
 
 				if (bWasLobbyMemberAdded)
 				{
-					TriggerOnSessionParticipantsChangeDelegates(Session->SessionName, *PlayerId, true);
+					TriggerOnSessionParticipantJoinedDelegates(Session->SessionName, *PlayerId);
 				}
 				else
 				{
@@ -2907,17 +2937,17 @@ void FOnlineSessionEOS::UpdateOrAddLobbyMember(const FUniqueNetIdEOSLobbyRef& Lo
 			}
 			else
 			{
-				UE_LOG_ONLINE(Warning, TEXT("[FOnlineSessionEOS::UpdateOrAddLobbyMember] EOS_LobbyDetails_CopyLobbyDetailsHandle not successful. Finished with EOS_EResult %s"), ANSI_TO_TCHAR(EOS_EResult_ToString(Result)));
+				UE_LOG_ONLINE_SESSION(Warning, TEXT("[FOnlineSessionEOS::UpdateOrAddLobbyMember] EOS_LobbyDetails_CopyLobbyDetailsHandle not successful. Finished with EOS_EResult %s"), ANSI_TO_TCHAR(EOS_EResult_ToString(Result)));
 			}
 		}
 		else
 		{
-			UE_LOG_ONLINE(Warning, TEXT("[FOnlineSessionEOS::UpdateOrAddLobbyMember] UniqueNetId %s not registered in the session's member settings."), *PlayerId->ToString());
+			UE_LOG_ONLINE_SESSION(Warning, TEXT("[FOnlineSessionEOS::UpdateOrAddLobbyMember] UniqueNetId %s not registered in the session's member settings."), *PlayerId->ToString());
 		}
 	}
 	else
 	{
-		UE_LOG_ONLINE(Warning, TEXT("[FOnlineSessionEOS::UpdateOrAddLobbyMember] Unable to retrieve session with LobbyId %s"), *LobbyNetId->ToString());
+		UE_LOG_ONLINE_SESSION(Warning, TEXT("[FOnlineSessionEOS::UpdateOrAddLobbyMember] Unable to retrieve session with LobbyId %s"), *LobbyNetId->ToString());
 	}
 }
 
@@ -3478,7 +3508,7 @@ void FOnlineSessionEOS::RemovePlayerFromSession(int32 LocalUserNum, FName Sessio
 	}
 	else
 	{
-		UE_LOG_ONLINE(Warning, TEXT("[FOnlineSessionEOS::RemovePlayerFromSession] Unable to retrieve session named %s"), *SessionName.ToString());
+		UE_LOG_ONLINE_SESSION(Warning, TEXT("[FOnlineSessionEOS::RemovePlayerFromSession] Unable to retrieve session named %s"), *SessionName.ToString());
 	}
 }
 
@@ -3701,7 +3731,7 @@ uint32 FOnlineSessionEOS::JoinLobbySession(int32 PlayerNum, FNamedOnlineSession*
 		}
 		else
 		{
-			UE_LOG_ONLINE(Warning, TEXT("[FOnlineSessionEOS::JoinLobbySession] SessionId not valid"));
+			UE_LOG_ONLINE_SESSION(Warning, TEXT("[FOnlineSessionEOS::JoinLobbySession] SessionId not valid"));
 		}
 	}
 	else
