Added logic to update cached Lobby information after Join operation

diff --git a/Engine/Plugins/Online/OnlineSubsystemEOS/Source/OnlineSubsystemEOS/Private/OnlineSessionEOS.cpp b/Engine/Plugins/Online/OnlineSubsystemEOS/Source/OnlineSubsystemEOS/Private/OnlineSessionEOS.cpp
--- a/Engine/Plugins/Online/OnlineSubsystemEOS/Source/OnlineSubsystemEOS/Private/OnlineSessionEOS.cpp
+++ b/Engine/Plugins/Online/OnlineSubsystemEOS/Source/OnlineSubsystemEOS/Private/OnlineSessionEOS.cpp
@@ -615,7 +615,7 @@ void FOnlineSessionEOS::RegisterLobbyNotifications()
 	LobbyUpdateReceivedCallback = LobbyUpdateReceivedCallbackObj;
 	LobbyUpdateReceivedCallbackObj->CallbackLambda = [this](const EOS_Lobby_LobbyUpdateReceivedCallbackInfo* Data)
 	{
-		OnLobbyUpdateReceived(Data->LobbyId);
+		OnLobbyUpdateReceived(Data->LobbyId, EOSSubsystem->UserManager->GetLocalProductUserId());
 	};
 
 	LobbyUpdateReceivedId = EOS_Lobby_AddNotifyLobbyUpdateReceived(LobbyHandle, &AddNotifyLobbyUpdateReceivedOptions, LobbyUpdateReceivedCallbackObj, LobbyUpdateReceivedCallbackObj->GetCallbackPtr());
@@ -675,7 +675,7 @@ void FOnlineSessionEOS::RegisterLobbyNotifications()
 	JoinLobbyAcceptedId = EOS_Lobby_AddNotifyJoinLobbyAccepted(LobbyHandle, &AddNotifyJoinLobbyAcceptedOptions, JoinLobbyAcceptedCallbackObj, JoinLobbyAcceptedCallbackObj->GetCallbackPtr());
 }
 
-void FOnlineSessionEOS::OnLobbyUpdateReceived(const EOS_LobbyId& LobbyId)
+void FOnlineSessionEOS::OnLobbyUpdateReceived(const EOS_LobbyId& LobbyId, const EOS_ProductUserId& LocalUserId)
 {
 	const FUniqueNetIdEOSLobbyRef LobbyNetId = FUniqueNetIdEOSLobby::Create(UTF8_TO_TCHAR(LobbyId));
 	FNamedOnlineSession* Session = GetNamedSessionFromLobbyId(*LobbyNetId);
@@ -684,7 +684,7 @@ void FOnlineSessionEOS::OnLobbyUpdateReceived(const EOS_LobbyId& LobbyId)
 		EOS_Lobby_CopyLobbyDetailsHandleOptions Options = {};
 		Options.ApiVersion = EOS_LOBBY_COPYLOBBYDETAILSHANDLE_API_LATEST;
 		Options.LobbyId = LobbyId;
-		Options.LocalUserId = EOSSubsystem->UserManager->GetLocalProductUserId();
+		Options.LocalUserId = LocalUserId;
 
 		EOS_HLobbyDetails LobbyDetailsHandle;
 
@@ -3680,8 +3680,8 @@ uint32 FOnlineSessionEOS::JoinLobbySession(int32 PlayerNum, FNamedOnlineSession*
  			TSharedRef<FLobbyDetailsEOS> LobbyDetails = LobbySearchResultsCache[Session->SessionInfo->GetSessionId().ToString()];
  			JoinLobbyOptions.LobbyDetailsHandle = LobbyDetails->LobbyDetailsHandle;
 
-			FName SessionName = Session->SessionName;
-			FUniqueNetIdPtr LocalUserNetId = EOSSubsystem->UserManager->GetLocalUniqueNetIdEOS(PlayerNum);
+			const FName SessionName = Session->SessionName;
+			const FUniqueNetIdEOSPtr LocalUserNetId = EOSSubsystem->UserManager->GetLocalUniqueNetIdEOS(PlayerNum);
 
 			FLobbyJoinedCallback* CallbackObj = new FLobbyJoinedCallback(FOnlineSessionEOSWeakPtr(AsShared()));
 			LobbyJoinedCallback = CallbackObj;
@@ -3690,7 +3690,7 @@ uint32 FOnlineSessionEOS::JoinLobbySession(int32 PlayerNum, FNamedOnlineSession*
 				FNamedOnlineSession* Session = GetNamedSession(SessionName);
 				if (Session)
 				{
-					bool bWasSuccessful = Data->ResultCode == EOS_EResult::EOS_Success;
+					const bool bWasSuccessful = Data->ResultCode == EOS_EResult::EOS_Success;
 					if (bWasSuccessful)
 					{
 						UE_LOG_ONLINE_SESSION(Verbose, TEXT("[FOnlineSessionEOS::JoinLobbySession] JoinLobby was successful. LobbyId is %d."), Data->LobbyId);
@@ -3703,6 +3703,8 @@ uint32 FOnlineSessionEOS::JoinLobbySession(int32 PlayerNum, FNamedOnlineSession*
 							VoiceChatUser->AddLobbyRoom(UTF8_TO_TCHAR(Data->LobbyId));
 						}
 #endif
+
+						OnLobbyUpdateReceived(Data->LobbyId, LocalUserNetId->GetProductUserId());
 					}
 					else
 					{
diff --git a/Engine/Plugins/Online/OnlineSubsystemEOS/Source/OnlineSubsystemEOS/Private/OnlineSessionEOS.h b/Engine/Plugins/Online/OnlineSubsystemEOS/Source/OnlineSubsystemEOS/Private/OnlineSessionEOS.h
--- a/Engine/Plugins/Online/OnlineSubsystemEOS/Source/OnlineSubsystemEOS/Private/OnlineSessionEOS.h
+++ b/Engine/Plugins/Online/OnlineSubsystemEOS/Source/OnlineSubsystemEOS/Private/OnlineSessionEOS.h
@@ -239,7 +239,7 @@ private:
 	EOS_NotificationId JoinLobbyAcceptedId;
 	FCallbackBase* JoinLobbyAcceptedCallback;
 
-	void OnLobbyUpdateReceived(const EOS_LobbyId& LobbyId);
+	void OnLobbyUpdateReceived(const EOS_LobbyId& LobbyId, const EOS_ProductUserId& LocalUserId);
 	void OnLobbyMemberUpdateReceived(const EOS_LobbyId& LobbyId, const EOS_ProductUserId& TargetUserId);
 	void OnMemberStatusReceived(const EOS_LobbyId& LobbyId, const EOS_ProductUserId& TargetUserId, EOS_ELobbyMemberStatus CurrentStatus);
 	void OnLobbyInviteAccepted(const char* InviteId, const EOS_ProductUserId& LocalUserId, const EOS_ProductUserId& TargetUserId);
