Removed Product Name from EOS BucketId and added logic to enforce 60 character limit

diff --git a/Engine/Plugins/Online/OnlineSubsystemEOS/Source/OnlineSubsystemEOS/Private/OnlineSessionEOS.cpp b/Engine/Plugins/Online/OnlineSubsystemEOS/Source/OnlineSubsystemEOS/Private/OnlineSessionEOS.cpp
--- a/Engine/Plugins/Online/OnlineSubsystemEOS/Source/OnlineSubsystemEOS/Private/OnlineSessionEOS.cpp
+++ b/Engine/Plugins/Online/OnlineSubsystemEOS/Source/OnlineSubsystemEOS/Private/OnlineSessionEOS.cpp
@@ -22,7 +22,7 @@
 	#include "eos_lobby.h"
 
 /** This is the game name plus version in ansi done once for optimization */
-char BucketIdAnsi[EOS_OSS_STRING_BUFFER_LENGTH];
+char BucketIdAnsi[EOS_OSS_BUCKET_ID_STRING_LENGTH + 1]; // X characters plus null terminator
 
 FString MakeStringFromAttributeValue(const EOS_Sessions_AttributeData* Attribute)
 {
diff --git a/Engine/Plugins/Online/OnlineSubsystemEOS/Source/OnlineSubsystemEOS/Private/OnlineSubsystemEOS.cpp b/Engine/Plugins/Online/OnlineSubsystemEOS/Source/OnlineSubsystemEOS/Private/OnlineSubsystemEOS.cpp
--- a/Engine/Plugins/Online/OnlineSubsystemEOS/Source/OnlineSubsystemEOS/Private/OnlineSubsystemEOS.cpp
+++ b/Engine/Plugins/Online/OnlineSubsystemEOS/Source/OnlineSubsystemEOS/Private/OnlineSubsystemEOS.cpp
@@ -409,8 +409,17 @@ bool FOnlineSubsystemEOS::Init()
 	UserManager = MakeShareable(new FUserManagerEOS(this));
 	UserManager->Init();
 	SessionInterfacePtr = MakeShareable(new FOnlineSessionEOS(this));
-	// Set the bucket id to use for all sessions based upon the name and version to avoid upgrade issues
-	SessionInterfacePtr->Init(EOSSDKManager->GetProductName() + TEXT("_") + FString::FromInt(GetBuildUniqueId()));
+	// Set the bucket id to use for all sessions based upon the Build Id to avoid upgrade issues
+	FString BucketIdStr = FString::FromInt(GetBuildUniqueId());
+	if (BucketIdStr.Len() > EOS_OSS_BUCKET_ID_STRING_LENGTH)
+	{
+		FString NewBucketIdStr = BucketIdStr.Left(EOS_OSS_BUCKET_ID_STRING_LENGTH);
+
+		UE_LOG_ONLINE(Warning, TEXT("[FOnlineSubsystemEOS::Init] Default BucketId [%s] is too long (%d characters) and will be shortened to fit the limit (%d characters). New BucketId is [%s]"), *BucketIdStr, BucketIdStr.Len(), EOS_OSS_BUCKET_ID_STRING_LENGTH, *NewBucketIdStr);
+
+		BucketIdStr = MoveTemp(NewBucketIdStr);		
+	}
+	SessionInterfacePtr->Init(BucketIdStr);
 	StatsInterfacePtr = MakeShareable(new FOnlineStatsEOS(this));
 	LeaderboardsInterfacePtr = MakeShareable(new FOnlineLeaderboardsEOS(this));
 	AchievementsInterfacePtr = MakeShareable(new FOnlineAchievementsEOS(this));
diff --git a/Engine/Plugins/Online/OnlineSubsystemEOS/Source/OnlineSubsystemEOS/Private/OnlineSubsystemEOSTypes.h b/Engine/Plugins/Online/OnlineSubsystemEOS/Source/OnlineSubsystemEOS/Private/OnlineSubsystemEOSTypes.h
--- a/Engine/Plugins/Online/OnlineSubsystemEOS/Source/OnlineSubsystemEOS/Private/OnlineSubsystemEOSTypes.h
+++ b/Engine/Plugins/Online/OnlineSubsystemEOS/Source/OnlineSubsystemEOS/Private/OnlineSubsystemEOSTypes.h
@@ -11,6 +11,7 @@
 #include "OnlineSubsystem.h"
 #include "OnlineSubsystemTypes.h"
 
+#define EOS_OSS_BUCKET_ID_STRING_LENGTH 60
 #define EOS_OSS_STRING_BUFFER_LENGTH 256 + 1 // 256 plus null terminator
 
 class FOnlineSubsystemEOS;
