- Fix CopyLobbyAttributes not updating existing keys.

diff --git a/Engine/Plugins/Online/OnlineSubsystemEOS/Source/OnlineSubsystemEOS/Private/OnlineSessionEOS.cpp b/Engine/Plugins/Online/OnlineSubsystemEOS/Source/OnlineSubsystemEOS/Private/OnlineSessionEOS.cpp
--- a/Engine/Plugins/Online/OnlineSubsystemEOS/Source/OnlineSubsystemEOS/Private/OnlineSessionEOS.cpp
+++ b/Engine/Plugins/Online/OnlineSubsystemEOS/Source/OnlineSubsystemEOS/Private/OnlineSessionEOS.cpp
@@ -4373,7 +4373,7 @@ void FOnlineSessionEOS::CopyLobbyAttributes(const TSharedRef<FLobbyDetailsEOS>&
 				}
 				}
 
-				OutSession.SessionSettings.Settings.FindOrAdd(FName(Key), Setting);
+				OutSession.SessionSettings.Settings.Emplace(FName(Key), MoveTemp(Setting));
 			}
 		}
 
