- Fix memory stomp caused by calling strncpy with a num greater than the length of the array.

diff --git a/Engine/Plugins/Online/OnlineSubsystemEOS/Source/OnlineSubsystemEOS/Private/OnlineSessionEOS.cpp b/Engine/Plugins/Online/OnlineSubsystemEOS/Source/OnlineSubsystemEOS/Private/OnlineSessionEOS.cpp
--- a/Engine/Plugins/Online/OnlineSubsystemEOS/Source/OnlineSubsystemEOS/Private/OnlineSessionEOS.cpp
+++ b/Engine/Plugins/Online/OnlineSubsystemEOS/Source/OnlineSubsystemEOS/Private/OnlineSessionEOS.cpp
@@ -446,7 +446,7 @@ FOnlineSessionEOS::~FOnlineSessionEOS()
 
 void FOnlineSessionEOS::Init(const FString& InBucketId)
 {
-	FCStringAnsi::Strncpy(BucketIdAnsi, TCHAR_TO_UTF8(*InBucketId), EOS_OSS_STRING_BUFFER_LENGTH);
+	FCStringAnsi::Strncpy(BucketIdAnsi, TCHAR_TO_UTF8(*InBucketId), sizeof(BucketIdAnsi));
 
 	// Register for session invite notifications
 	FSessionInviteAcceptedCallback* SessionInviteAcceptedCallbackObj = new FSessionInviteAcceptedCallback(FOnlineSessionEOSWeakPtr(AsShared()));
