Added logic to avoid setting CacheDir parameter on EOSSDK platform creation if said cache is not mounted

diff --git a/Engine/Plugins/Online/EOSShared/Source/EOSShared/Private/EOSSDKManager.cpp b/Engine/Plugins/Online/EOSShared/Source/EOSShared/Private/EOSSDKManager.cpp
--- a/Engine/Plugins/Online/EOSShared/Source/EOSShared/Private/EOSSDKManager.cpp
+++ b/Engine/Plugins/Online/EOSShared/Source/EOSShared/Private/EOSSDKManager.cpp
@@ -411,7 +411,15 @@ IEOSPlatformHandlePtr FEOSSDKManager::CreatePlatform(const FString& PlatformConf
 	if (PlatformConfig->bDisableOverlay) PlatformOptions.Flags |= EOS_PF_DISABLE_OVERLAY;
 	if (PlatformConfig->bDisableSocialOverlay) PlatformOptions.Flags |= EOS_PF_DISABLE_SOCIAL_OVERLAY;
 
-	PlatformOptions.CacheDirectory = Utf8CacheDirectory.Length() ? Utf8DeploymentId.Get() : nullptr;
+	if (FPlatformMisc::IsCacheStorageAvailable())
+	{
+		PlatformOptions.CacheDirectory = Utf8CacheDirectory.Length() ? Utf8DeploymentId.Get() : nullptr;
+	}
+	else
+	{
+		PlatformOptions.CacheDirectory = nullptr;
+	}
+
 	PlatformOptions.TickBudgetInMilliseconds = PlatformConfig->TickBudgetInMilliseconds;
 
 	static_assert(EOS_PLATFORM_RTCOPTIONS_API_LATEST == 1, "EOS_Platform_RTCOptions updated");
diff --git a/Engine/Plugins/Online/OnlineServicesEOSGS/Source/Private/Online/OnlineServicesEOSGSPlatformFactory.cpp b/Engine/Plugins/Online/OnlineServicesEOSGS/Source/Private/Online/OnlineServicesEOSGSPlatformFactory.cpp
--- a/Engine/Plugins/Online/OnlineServicesEOSGS/Source/Private/Online/OnlineServicesEOSGSPlatformFactory.cpp
+++ b/Engine/Plugins/Online/OnlineServicesEOSGS/Source/Private/Online/OnlineServicesEOSGSPlatformFactory.cpp
@@ -118,7 +118,15 @@ IEOSPlatformHandlePtr FOnlineServicesEOSGSPlatformFactory::CreatePlatform()
 	PlatformOptions.ClientCredentials.ClientId = ClientId.Get();
 	PlatformOptions.ClientCredentials.ClientSecret = ClientSecret.Get();
 	PlatformOptions.EncryptionKey = EncryptionKey.Get();
-	PlatformOptions.CacheDirectory = CacheDirectory.Get();
+
+	if (FPlatformMisc::IsCacheStorageAvailable())
+	{
+		PlatformOptions.CacheDirectory = CacheDirectory.Get();
+	}
+	else
+	{
+		PlatformOptions.CacheDirectory = nullptr;
+	}
 
 	IEOSPlatformHandlePtr EOSPlatformHandle = SDKManager->CreatePlatform(PlatformOptions);
 	if (!EOSPlatformHandle)
diff --git a/Engine/Plugins/Online/OnlineSubsystemEOS/Source/OnlineSubsystemEOS/Private/OnlineSubsystemEOS.cpp b/Engine/Plugins/Online/OnlineSubsystemEOS/Source/OnlineSubsystemEOS/Private/OnlineSubsystemEOS.cpp
--- a/Engine/Plugins/Online/OnlineSubsystemEOS/Source/OnlineSubsystemEOS/Private/OnlineSubsystemEOS.cpp
+++ b/Engine/Plugins/Online/OnlineSubsystemEOS/Source/OnlineSubsystemEOS/Private/OnlineSubsystemEOS.cpp
@@ -238,8 +238,16 @@ bool FOnlineSubsystemEOS::PlatformCreate()
 	PlatformOptions.Flags = bOverlayAllowed ? OverlayFlags : EOS_PF_DISABLE_OVERLAY;
 	// Make the cache directory be in the user's writable area
 
-	const FString CacheDir = EOSSDKManager->GetCacheDirBase() / ArtifactName / EOSSettings.CacheDir;
-	FCStringAnsi::Strncpy(PlatformOptions.CacheDirectoryAnsi, TCHAR_TO_UTF8(*CacheDir), EOS_OSS_STRING_BUFFER_LENGTH);
+	if (FPlatformMisc::IsCacheStorageAvailable())
+	{
+		const FString CacheDir = EOSSDKManager->GetCacheDirBase() / ArtifactName / EOSSettings.CacheDir;
+		FCStringAnsi::Strncpy(PlatformOptions.CacheDirectoryAnsi, TCHAR_TO_UTF8(*CacheDir), EOS_OSS_STRING_BUFFER_LENGTH);
+	}
+	else
+	{
+		PlatformOptions.CacheDirectory = nullptr;
+	}
+
 	FCStringAnsi::Strncpy(PlatformOptions.EncryptionKeyAnsi, TCHAR_TO_UTF8(*ArtifactSettings.EncryptionKey), EOS_ENCRYPTION_KEY_MAX_BUFFER_LEN);
 
 #if WITH_EOS_RTC
diff --git a/Engine/Source/Runtime/Core/Public/GenericPlatform/GenericPlatformMisc.h b/Engine/Source/Runtime/Core/Public/GenericPlatform/GenericPlatformMisc.h
--- a/Engine/Source/Runtime/Core/Public/GenericPlatform/GenericPlatformMisc.h
+++ b/Engine/Source/Runtime/Core/Public/GenericPlatform/GenericPlatformMisc.h
@@ -784,6 +784,16 @@ struct CORE_API FGenericPlatformMisc
 		return true;
 	}
 
+	/**
+	 * For platforms that support cache storage functionality, setting a directory for it when it's not available can result in errors.
+	 * This method's purpose is to be overridden in such platforms, with logic that checks if said cache is available.
+	 * In all other platforms without this constraint, it should return true by default.
+	 */
+	static bool IsCacheStorageAvailable()
+	{
+		return true;
+	}
+
 	static bool SupportsLocalCaching()
 	{
 		return true;
