- Add support for -epicsandboxid command line argument from EGS. This enables duplicate settings entries for a given artifact id (passed with -epicapp), and we will select the one with matching SandboxId if the command line arg is present.

diff --git a/Engine/Plugins/Online/OnlineSubsystemEOS/Source/OnlineSubsystemEOS/Private/EOSSettings.cpp b/Engine/Plugins/Online/OnlineSubsystemEOS/Source/OnlineSubsystemEOS/Private/EOSSettings.cpp
--- a/Engine/Plugins/Online/OnlineSubsystemEOS/Source/OnlineSubsystemEOS/Private/EOSSettings.cpp
+++ b/Engine/Plugins/Online/OnlineSubsystemEOS/Source/OnlineSubsystemEOS/Private/EOSSettings.cpp
@@ -79,8 +79,10 @@ inline FString StripQuotes(const FString& Source)
 	return Source;
 }
 
-void FEOSArtifactSettings::ParseRawArrayEntry(const FString& RawLine)
+FEOSArtifactSettings ParseArtifactSettingsFromConfigString(const FString& RawLine)
 {
+	FEOSArtifactSettings Result;
+
 	const TCHAR* Delims[4] = { TEXT("("), TEXT(")"), TEXT("="), TEXT(",") };
 	TArray<FString> Values;
 	RawLine.ParseIntoArray(Values, Delims, 4, false);
@@ -94,34 +96,36 @@ void FEOSArtifactSettings::ParseRawArrayEntry(const FString& RawLine)
 		// Parse which struct field
 		if (Values[ValueIndex] == TEXT("ArtifactName"))
 		{
-			ArtifactName = StripQuotes(Values[ValueIndex + 1]);
+			Result.ArtifactName = StripQuotes(Values[ValueIndex + 1]);
 		}
 		else if (Values[ValueIndex] == TEXT("ClientId"))
 		{
-			ClientId = StripQuotes(Values[ValueIndex + 1]);
+			Result.ClientId = StripQuotes(Values[ValueIndex + 1]);
 		}
 		else if (Values[ValueIndex] == TEXT("ClientSecret"))
 		{
-			ClientSecret = StripQuotes(Values[ValueIndex + 1]);
+			Result.ClientSecret = StripQuotes(Values[ValueIndex + 1]);
 		}
 		else if (Values[ValueIndex] == TEXT("ProductId"))
 		{
-			ProductId = StripQuotes(Values[ValueIndex + 1]);
+			Result.ProductId = StripQuotes(Values[ValueIndex + 1]);
 		}
 		else if (Values[ValueIndex] == TEXT("SandboxId"))
 		{
-			SandboxId = StripQuotes(Values[ValueIndex + 1]);
+			Result.SandboxId = StripQuotes(Values[ValueIndex + 1]);
 		}
 		else if (Values[ValueIndex] == TEXT("DeploymentId"))
 		{
-			DeploymentId = StripQuotes(Values[ValueIndex + 1]);
+			Result.DeploymentId = StripQuotes(Values[ValueIndex + 1]);
 		}
 		else if (Values[ValueIndex] == TEXT("EncryptionKey"))
 		{
-			EncryptionKey = StripQuotes(Values[ValueIndex + 1]);
+			Result.EncryptionKey = StripQuotes(Values[ValueIndex + 1]);
 		}
 		ValueIndex++;
 	}
+
+	return Result;
 }
 
 FEOSSettings UEOSSettings::GetSettings()
@@ -194,99 +198,99 @@ FEOSSettings UEOSSettings::ToNative() const
 	return Native;
 }
 
-bool UEOSSettings::GetSettingsForArtifact(const FString& ArtifactName, FEOSArtifactSettings& OutSettings)
+bool UEOSSettings::GetSelectedArtifactSettings(FEOSArtifactSettings& OutSettings)
 {
-	if (UObjectInitialized())
+	// Get default artifact name from config.
+	FString ArtifactName = GetDefaultArtifactName();
+	// Prefer -epicapp over config. This generally comes from EGS.
+	FParse::Value(FCommandLine::Get(), TEXT("EpicApp="), ArtifactName);
+	// Prefer -EOSArtifactNameOverride over previous.
+	FParse::Value(FCommandLine::Get(), TEXT("EOSArtifactNameOverride="), ArtifactName);
+
+	// Get the sandbox id, which only comes in via EGS. If present, grab the settings where both match.
+	FString SandboxId;
+	if (FParse::Value(FCommandLine::Get(), TEXT("EpicSandboxId="), SandboxId))
 	{
-		return UEOSSettings::AutoGetSettingsForArtifact(ArtifactName, OutSettings);
+		const bool bSuccess = GetArtifactSettings(ArtifactName, SandboxId, OutSettings);
+		UE_CLOG_ONLINE(!bSuccess, Error, TEXT("UEOSSettings::GetSelectedArtifactSettings() ArtifactName=[%s] SandboxId=[%s] no settings found."), *ArtifactName, *SandboxId);
+		return bSuccess;
 	}
-	return UEOSSettings::ManualGetSettingsForArtifact(ArtifactName, OutSettings);
+
+	// Fall back on just matching the Artifact name. This assumes non-EGS and only one settings entry per ArtifactName in config.
+	const bool bSuccess = GetArtifactSettings(ArtifactName, OutSettings);
+	UE_CLOG_ONLINE(!bSuccess, Error, TEXT("UEOSSettings::GetSelectedArtifactSettings() ArtifactName=[%s] no settings found."), *ArtifactName);
+	return bSuccess;
 }
 
-bool UEOSSettings::ManualGetSettingsForArtifact(const FString& ArtifactName, FEOSArtifactSettings& OutSettings)
+FString UEOSSettings::GetDefaultArtifactName()
 {
-	static TOptional<FString> CachedDefaultArtifactName;
-	static TOptional<TArray<FEOSArtifactSettings>> CachedArtifactSettings;
-
-	if (!CachedDefaultArtifactName.IsSet())
+	if (UObjectInitialized())
 	{
-		CachedDefaultArtifactName.Emplace();
-
-		GConfig->GetString(INI_SECTION, TEXT("DefaultArtifactName"), *CachedDefaultArtifactName, GEngineIni);
+		return GetDefault<UEOSSettings>()->DefaultArtifactName;
 	}
 
-	if (!CachedArtifactSettings.IsSet())
-	{
-		CachedArtifactSettings.Emplace();
+	FString DefaultArtifactName;
+	GConfig->GetString(INI_SECTION, TEXT("DefaultArtifactName"), DefaultArtifactName, GEngineIni);
+	return DefaultArtifactName;
+}
 
-		TArray<FString> Artifacts;
-		GConfig->GetArray(INI_SECTION, TEXT("Artifacts"), Artifacts, GEngineIni);
-		for (const FString& Line : Artifacts)
-		{
-			FEOSArtifactSettings Artifact;
-			Artifact.ParseRawArrayEntry(Line);
-			CachedArtifactSettings->Add(Artifact);
-		}
-	}
+bool UEOSSettings::GetArtifactSettings(const FString& ArtifactName, FEOSArtifactSettings& OutSettings)
+{
+	return GetArtifactSettingsImpl(ArtifactName, TOptional<FString>(), OutSettings);
+}
 
-	FString ArtifactNameOverride;
-	// Figure out which config object we are loading
-	FParse::Value(FCommandLine::Get(), TEXT("EOSArtifactNameOverride="), ArtifactNameOverride);
-	if (ArtifactNameOverride.IsEmpty())
-	{
-		ArtifactNameOverride = ArtifactName;
-	}
+bool UEOSSettings::GetArtifactSettings(const FString& ArtifactName, const FString& SandboxId, FEOSArtifactSettings& OutSettings)
+{
+	return GetArtifactSettingsImpl(ArtifactName, SandboxId, OutSettings);
+}
 
-	// Search by name and then default if not found
-	for (const FEOSArtifactSettings& Artifact : *CachedArtifactSettings)
+bool UEOSSettings::GetArtifactSettingsImpl(const FString& ArtifactName, const TOptional<FString>& SandboxId, FEOSArtifactSettings& OutSettings)
+{
+	if (UObjectInitialized())
 	{
-		if (Artifact.ArtifactName == ArtifactNameOverride)
+		const UEOSSettings* This = GetDefault<UEOSSettings>();
+		const FArtifactSettings* Found = This->Artifacts.FindByPredicate([&ArtifactName, &SandboxId](const FArtifactSettings& Element)
 		{
-			OutSettings = Artifact;
+			return Element.ArtifactName == ArtifactName
+				&& (!SandboxId.IsSet() || Element.SandboxId == SandboxId);
+		});
+		if (Found)
+		{
+			OutSettings = Found->ToNative();
 			return true;
 		}
+		return false;
 	}
-
-	for (const FEOSArtifactSettings& Artifact : *CachedArtifactSettings)
+	else
 	{
-		if (Artifact.ArtifactName == *CachedDefaultArtifactName)
+		const TArray<FEOSArtifactSettings>& CachedSettings = GetCachedArtifactSettings();
+		const FEOSArtifactSettings* Found = CachedSettings.FindByPredicate([&ArtifactName, &SandboxId](const FEOSArtifactSettings& Element)
+		{
+			return Element.ArtifactName == ArtifactName
+				&& (!SandboxId.IsSet() || Element.SandboxId == SandboxId);
+		});
+		if (Found)
 		{
-			OutSettings = Artifact;
+			OutSettings = *Found;
 			return true;
 		}
+		return false;
 	}
-
-	return false;
 }
 
-bool UEOSSettings::AutoGetSettingsForArtifact(const FString& ArtifactName, FEOSArtifactSettings& OutSettings)
+const TArray<FEOSArtifactSettings>& UEOSSettings::GetCachedArtifactSettings()
 {
-	const UEOSSettings* This = GetDefault<UEOSSettings>();
-	FString ArtifactNameOverride;
-	// Figure out which config object we are loading
-	FParse::Value(FCommandLine::Get(), TEXT("EOSArtifactNameOverride="), ArtifactNameOverride);
-	if (ArtifactNameOverride.IsEmpty())
+	static TArray<FEOSArtifactSettings> CachedArtifactSettings = []()
 	{
-		ArtifactNameOverride = ArtifactName;
-	}
-	for (const FArtifactSettings& Artifact : This->Artifacts)
-	{
-		if (Artifact.ArtifactName == ArtifactNameOverride)
-		{
-			OutSettings = Artifact.ToNative();
-			return true;
-		}
-	}
-	for (const FArtifactSettings& Artifact : This->Artifacts)
-	{
-		if (Artifact.ArtifactName == This->DefaultArtifactName)
-		{
-			OutSettings = Artifact.ToNative();
-			return true;
-		}
-	}
-	UE_LOG_ONLINE(Error, TEXT("UEOSSettings::AutoGetSettingsForArtifact() failed due to missing config object specified. Check your project settings"));
-	return false;
+		TArray<FString> ConfigArray;
+		GConfig->GetArray(INI_SECTION, TEXT("Artifacts"), ConfigArray, GEngineIni);
+
+		TArray<FEOSArtifactSettings> Result;
+		Algo::Transform(ConfigArray, Result, &ParseArtifactSettingsFromConfigString);
+
+		return Result;
+	}();
+	return CachedArtifactSettings;
 }
 
 #if WITH_EDITOR
diff --git a/Engine/Plugins/Online/OnlineSubsystemEOS/Source/OnlineSubsystemEOS/Private/OnlineSubsystemEOS.cpp b/Engine/Plugins/Online/OnlineSubsystemEOS/Source/OnlineSubsystemEOS/Private/OnlineSubsystemEOS.cpp
--- a/Engine/Plugins/Online/OnlineSubsystemEOS/Source/OnlineSubsystemEOS/Private/OnlineSubsystemEOS.cpp
+++ b/Engine/Plugins/Online/OnlineSubsystemEOS/Source/OnlineSubsystemEOS/Private/OnlineSubsystemEOS.cpp
@@ -195,13 +195,10 @@ void FOnlineSubsystemEOS::ModuleShutdown()
 /** Common method for creating the EOS platform */
 bool FOnlineSubsystemEOS::PlatformCreate()
 {
-	FString ArtifactName;
-	FParse::Value(FCommandLine::Get(), TEXT("EpicApp="), ArtifactName);
-	// Find the settings for this artifact
 	FEOSArtifactSettings ArtifactSettings;
-	if (!UEOSSettings::GetSettingsForArtifact(ArtifactName, ArtifactSettings))
+	if (!UEOSSettings::GetSelectedArtifactSettings(ArtifactSettings))
 	{
-		UE_LOG_ONLINE(Error, TEXT("FOnlineSubsystemEOS::PlatformCreate() failed to find artifact settings object for artifact (%s)"), *ArtifactName);
+		UE_LOG_ONLINE(Error, TEXT("FOnlineSubsystemEOS::PlatformCreate() GetSelectedArtifactSettings failed"));
 		return false;
 	}
 
@@ -240,7 +237,7 @@ bool FOnlineSubsystemEOS::PlatformCreate()
 
 	if (FPlatformMisc::IsCacheStorageAvailable())
 	{
-		const FString CacheDir = EOSSDKManager->GetCacheDirBase() / ArtifactName / EOSSettings.CacheDir;
+		const FString CacheDir = EOSSDKManager->GetCacheDirBase() / ArtifactSettings.ArtifactName / EOSSettings.CacheDir;
 		FCStringAnsi::Strncpy(PlatformOptions.CacheDirectoryAnsi, TCHAR_TO_UTF8(*CacheDir), EOS_OSS_STRING_BUFFER_LENGTH);
 	}
 	else
@@ -402,16 +399,16 @@ bool FOnlineSubsystemEOS::Init()
 	}
 
 	// We set the product id
-	FString ArtifactName;
-	FParse::Value(FCommandLine::Get(), TEXT("EpicApp="), ArtifactName);
 	FEOSArtifactSettings ArtifactSettings;
-	if (UEOSSettings::GetSettingsForArtifact(ArtifactName, ArtifactSettings))
+	if (UEOSSettings::GetSelectedArtifactSettings(ArtifactSettings))
 	{
 		ProductId = ArtifactSettings.ProductId;
 	}
 	else
 	{
-		UE_LOG_ONLINE(Warning, TEXT("[FOnlineSubsystemEOS::Init] Failed to find artifact settings object for artifact (%s). ProductIdAnsi not set."), *ArtifactName);
+		// This really should not be possible, if we made it past PlatformCreate.
+		checkNoEntry();
+		UE_LOG_ONLINE(Warning, TEXT("[FOnlineSubsystemEOS::Init] GetSelectedArtifactSettings failed, ProductIdAnsi not set."));
 	}
 
 	UserManager = MakeShareable(new FUserManagerEOS(this));
diff --git a/Engine/Plugins/Online/OnlineSubsystemEOS/Source/OnlineSubsystemEOS/Public/EOSSettings.h b/Engine/Plugins/Online/OnlineSubsystemEOS/Source/OnlineSubsystemEOS/Public/EOSSettings.h
--- a/Engine/Plugins/Online/OnlineSubsystemEOS/Source/OnlineSubsystemEOS/Public/EOSSettings.h
+++ b/Engine/Plugins/Online/OnlineSubsystemEOS/Source/OnlineSubsystemEOS/Public/EOSSettings.h
@@ -17,8 +17,6 @@ struct FEOSArtifactSettings
 	FString SandboxId;
 	FString DeploymentId;
 	FString EncryptionKey;
-
-	void ParseRawArrayEntry(const FString& RawLine);
 };
 
 UCLASS(Deprecated)
@@ -165,8 +163,8 @@ public:
 	UPROPERTY(Config, EditAnywhere, BlueprintReadOnly, Category="Crossplay Settings")
 	bool bMirrorPresenceToEAS = false;
 
-	/** Find the Settings for an artifact by name */
-	static bool GetSettingsForArtifact(const FString& ArtifactName, FEOSArtifactSettings& OutSettings);
+	/** Get the settings for the selected artifact */
+	static bool GetSelectedArtifactSettings(FEOSArtifactSettings& OutSettings);
 
 	static FEOSSettings GetSettings();
 	FEOSSettings ToNative() const;
@@ -175,8 +173,14 @@ private:
 #if WITH_EDITOR
 	virtual void PostEditChangeProperty(FPropertyChangedEvent& PropertyChangedEvent) override;
 #endif
-	static bool AutoGetSettingsForArtifact(const FString& ArtifactName, FEOSArtifactSettings& OutSettings);
-	static bool ManualGetSettingsForArtifact(const FString& ArtifactName, FEOSArtifactSettings& OutSettings);
+
+	static FString GetDefaultArtifactName();
+
+	static bool GetArtifactSettings(const FString& ArtifactName, FEOSArtifactSettings& OutSettings);
+	static bool GetArtifactSettings(const FString& ArtifactName, const FString& SandboxId, FEOSArtifactSettings& OutSettings);
+	static bool GetArtifactSettingsImpl(const FString& ArtifactName, const TOptional<FString>& SandboxId, FEOSArtifactSettings& OutSettings);
+
+	static const TArray<FEOSArtifactSettings>& GetCachedArtifactSettings();
 
 	static FEOSSettings AutoGetSettings();
 	static const FEOSSettings& ManualGetSettings();
