Add accessors to get Epic Account Id and Product User Id from OnlineSubsystemEOS net ids

diff --git a/Engine/Plugins/Online/OnlineSubsystemEOS/Source/OnlineSubsystemEOS/OnlineSubsystemEOS.Build.cs b/Engine/Plugins/Online/OnlineSubsystemEOS/Source/OnlineSubsystemEOS/OnlineSubsystemEOS.Build.cs
--- a/Engine/Plugins/Online/OnlineSubsystemEOS/Source/OnlineSubsystemEOS/OnlineSubsystemEOS.Build.cs
+++ b/Engine/Plugins/Online/OnlineSubsystemEOS/Source/OnlineSubsystemEOS/OnlineSubsystemEOS.Build.cs
@@ -11,7 +11,7 @@ public class OnlineSubsystemEOS : ModuleRules
 
 		PublicDependencyModuleNames.AddRange(
 			new string[] {
-				"OnlineSubsystemUtils"
+				"EOSSDK",
 			}
 		);
 
@@ -21,12 +21,12 @@ public class OnlineSubsystemEOS : ModuleRules
 				"CoreOnline",
 				"CoreUObject",
 				"Engine",
-				"EOSSDK",
 				"EOSShared",
 				"EOSVoiceChat",
 				"Json",
 				"OnlineBase",
 				"OnlineSubsystem",
+				"OnlineSubsystemUtils",
 				"Sockets",
 				"SocketSubsystemEOS",
 				"VoiceChat",
diff --git a/Engine/Plugins/Online/OnlineSubsystemEOS/Source/OnlineSubsystemEOS/Private/OnlineSubsystemEOSTypes.h b/Engine/Plugins/Online/OnlineSubsystemEOS/Source/OnlineSubsystemEOS/Private/OnlineSubsystemEOSTypes.h
--- a/Engine/Plugins/Online/OnlineSubsystemEOS/Source/OnlineSubsystemEOS/Private/OnlineSubsystemEOSTypes.h
+++ b/Engine/Plugins/Online/OnlineSubsystemEOS/Source/OnlineSubsystemEOS/Private/OnlineSubsystemEOSTypes.h
@@ -2,6 +2,7 @@
 
 #pragma once
 
+#include "OnlineSubsystemEOSTypesPublic.h"
 #include "Interfaces/OnlineIdentityInterface.h"
 #include "Interfaces/OnlinePresenceInterface.h"
 #include "Interfaces/OnlineUserInterface.h"
@@ -27,7 +28,7 @@ typedef TSharedRef<const class FUniqueNetIdEOS> FUniqueNetIdEOSRef;
 /**
  * Unique net id wrapper for a EOS account ids.
  */
-class FUniqueNetIdEOS : public FUniqueNetId
+class FUniqueNetIdEOS : public IUniqueNetIdEOS
 {
 public:
 	static const FUniqueNetIdEOS& Cast(const FUniqueNetId& NetId);
@@ -45,12 +46,12 @@ public:
 	virtual FString ToString() const override;
 	virtual FString ToDebugString() const override;
 
-	const EOS_EpicAccountId GetEpicAccountId() const
+	virtual const EOS_EpicAccountId GetEpicAccountId() const override
 	{
 		return EpicAccountId;
 	}
 
-	const EOS_ProductUserId GetProductUserId() const
+	virtual const EOS_ProductUserId GetProductUserId() const override
 	{
 		return ProductUserId;
 	}
diff --git a/Engine/Plugins/Online/OnlineSubsystemEOS/Source/OnlineSubsystemEOS/Public/OnlineSubsystemEOSTypesPublic.h b/Engine/Plugins/Online/OnlineSubsystemEOS/Source/OnlineSubsystemEOS/Public/OnlineSubsystemEOSTypesPublic.h
new file mode 100644
--- /dev/null
+++ b/Engine/Plugins/Online/OnlineSubsystemEOS/Source/OnlineSubsystemEOS/Public/OnlineSubsystemEOSTypesPublic.h
@@ -0,0 +1,31 @@
+// Copyright Epic Games, Inc. All Rights Reserved.
+
+#pragma once
+
+#if WITH_EOS_SDK
+
+#include "Online/CoreOnline.h"
+
+#if defined(EOS_PLATFORM_BASE_FILE_NAME)
+#include EOS_PLATFORM_BASE_FILE_NAME
+#endif
+#include "eos_base.h"
+#include "eos_common.h"
+
+/** Unique net id for OnlineSubsystemEOS */
+class IUniqueNetIdEOS : public FUniqueNetId
+{
+public:
+	/**
+	 * Get the Epic Account Id. May be null.
+	 * @return the Epic Account Id, or null.
+	 */
+	virtual const EOS_EpicAccountId GetEpicAccountId() const = 0;
+	/**
+	 * Get the Product User Id. May be null.
+	 * @return the Product User Id, or null.
+	 */
+	virtual const EOS_ProductUserId GetProductUserId() const = 0;
+};
+
+#endif // WITH_EOS_SDK
\ No newline at end of file
