Removed Recent and Blocked players from UserManagerEOS. Functionality is not supported and types were not used.

diff --git a/Engine/Plugins/Online/OnlineSubsystemEOS/Source/OnlineSubsystemEOS/Private/OnlineSubsystemEOSTypes.h b/Engine/Plugins/Online/OnlineSubsystemEOS/Source/OnlineSubsystemEOS/Private/OnlineSubsystemEOSTypes.h
--- a/Engine/Plugins/Online/OnlineSubsystemEOS/Source/OnlineSubsystemEOS/Private/OnlineSubsystemEOSTypes.h
+++ b/Engine/Plugins/Online/OnlineSubsystemEOS/Source/OnlineSubsystemEOS/Private/OnlineSubsystemEOSTypes.h
@@ -295,52 +295,6 @@ protected:
 	EInviteStatus::Type InviteStatus;
 };
 
-/**
- * Implementation of FOnlineBlockedPlayer methods that adds in the online user template to complete the interface
- */
-template<class BaseClass>
-class TOnlineBlockedPlayerEOS :
-	public TOnlineUserEOS<BaseClass, IAttributeAccessInterface>
-{
-public:
-	TOnlineBlockedPlayerEOS(const FUniqueNetIdEOSRef& InNetIdRef)
-		: TOnlineUserEOS<BaseClass, IAttributeAccessInterface>(InNetIdRef)
-	{
-	}
-};
-
-/**
- * Implementation of FOnlineRecentPlayer methods that adds in the online user template to complete the interface
- */
-template<class BaseClass>
-class TOnlineRecentPlayerEOS :
-	public TOnlineUserEOS<BaseClass, IAttributeAccessInterface>
-{
-public:
-	TOnlineRecentPlayerEOS(const FUniqueNetIdEOSRef& InNetIdRef)
-		: TOnlineUserEOS<BaseClass, IAttributeAccessInterface>(InNetIdRef)
-	{
-	}
-
-// FOnlineRecentPlayer
-	/**
-	 * @return last time the player was seen by the current user
-	 */
-	virtual FDateTime GetLastSeen() const override
-	{
-		return LastSeenTime;
-	}
-//~FOnlineRecentPlayer
-
-	void SetLastSeen(const FDateTime& InLastSeenTime)
-	{
-		LastSeenTime = InLastSeenTime;
-	}
-
-protected:
-	FDateTime LastSeenTime;
-};
-
 /** Class to handle all callbacks generically using a lambda to process callback results */
 template<typename CallbackFuncType, typename CallbackType, typename OwningType>
 class TEOSCallback :
diff --git a/Engine/Plugins/Online/OnlineSubsystemEOS/Source/OnlineSubsystemEOS/Private/UserManagerEOS.cpp b/Engine/Plugins/Online/OnlineSubsystemEOS/Source/OnlineSubsystemEOS/Private/UserManagerEOS.cpp
--- a/Engine/Plugins/Online/OnlineSubsystemEOS/Source/OnlineSubsystemEOS/Private/UserManagerEOS.cpp
+++ b/Engine/Plugins/Online/OnlineSubsystemEOS/Source/OnlineSubsystemEOS/Private/UserManagerEOS.cpp
@@ -1069,15 +1069,6 @@ void FUserManagerEOS::AddLocalUser(int32 LocalUserNum, EOS_EpicAccountId EpicAcc
 	NetIdStringToFriendsListMap.Emplace(NetId, FriendsList);
 	ReadFriendsList(LocalUserNum, FString());
 
-	FBlockedPlayersListEOSRef BlockedPlayersList = MakeShareable(new FBlockedPlayersListEOS(LocalUserNum, UserNetId));
-	LocalUserNumToBlockedPlayerListMap.Emplace(LocalUserNum, BlockedPlayersList);
-	NetIdStringToBlockedPlayerListMap.Emplace(NetId, BlockedPlayersList);
-	QueryBlockedPlayers(*UserNetId);
-
-	FRecentPlayersListEOSRef RecentPlayersList = MakeShareable(new FRecentPlayersListEOS(LocalUserNum, UserNetId));
-	LocalUserNumToRecentPlayerListMap.Emplace(LocalUserNum, RecentPlayersList);
-	NetIdStringToRecentPlayerListMap.Emplace(NetId, RecentPlayersList);
-
 	// Get auth token info
 	EOS_Auth_Token* AuthToken = nullptr;
 	EOS_Auth_CopyUserAuthTokenOptions Options = { };
diff --git a/Engine/Plugins/Online/OnlineSubsystemEOS/Source/OnlineSubsystemEOS/Private/UserManagerEOS.h b/Engine/Plugins/Online/OnlineSubsystemEOS/Source/OnlineSubsystemEOS/Private/UserManagerEOS.h
--- a/Engine/Plugins/Online/OnlineSubsystemEOS/Source/OnlineSubsystemEOS/Private/UserManagerEOS.h
+++ b/Engine/Plugins/Online/OnlineSubsystemEOS/Source/OnlineSubsystemEOS/Private/UserManagerEOS.h
@@ -69,40 +69,6 @@ public:
 typedef TSharedPtr<FOnlineFriendEOS> FOnlineFriendEOSPtr;
 typedef TSharedRef<FOnlineFriendEOS> FOnlineFriendEOSRef;
 
-/**
- * Concrete blocked player class built from the reusable templates
- */
-class FOnlineBlockedPlayerEOS :
-	public TOnlineBlockedPlayerEOS<FOnlineBlockedPlayer>
-{
-public:
-	FOnlineBlockedPlayerEOS(FUniqueNetIdEOSRef InUserId) :
-		TOnlineBlockedPlayerEOS<FOnlineBlockedPlayer>(InUserId)
-	{
-	}
-	virtual ~FOnlineBlockedPlayerEOS() = default;
-};
-
-typedef TSharedPtr<FOnlineBlockedPlayerEOS> FOnlineBlockedPlayerEOSPtr;
-typedef TSharedRef<FOnlineBlockedPlayerEOS> FOnlineBlockedPlayerEOSRef;
-
-/**
- * Concrete blocked player class built from the reusable templates
- */
-class FOnlineRecentPlayerEOS :
-	public TOnlineRecentPlayerEOS<FOnlineRecentPlayer>
-{
-public:
-	FOnlineRecentPlayerEOS(FUniqueNetIdEOSRef InUserId) :
-		TOnlineRecentPlayerEOS<FOnlineRecentPlayer>(InUserId)
-	{
-	}
-	virtual ~FOnlineRecentPlayerEOS() = default;
-};
-
-typedef TSharedPtr<FOnlineRecentPlayerEOS> FOnlineRecentPlayerEOSPtr;
-typedef TSharedRef<FOnlineRecentPlayerEOS> FOnlineRecentPlayerEOSRef;
-
 template<class ListClass, class ListClassReturnType>
 class TOnlinePlayerList
 {
@@ -189,32 +155,6 @@ public:
 
 typedef TSharedRef<FFriendsListEOS> FFriendsListEOSRef;
 
-class FBlockedPlayersListEOS :
-	public TOnlinePlayerList<FOnlineBlockedPlayerEOSRef, FOnlineBlockedPlayerEOSPtr>
-{
-public:
-	FBlockedPlayersListEOS(int32 InLocalUserNum, FUniqueNetIdEOSRef InOwningNetId)
-		: TOnlinePlayerList<FOnlineBlockedPlayerEOSRef, FOnlineBlockedPlayerEOSPtr>(InLocalUserNum, InOwningNetId)
-	{
-	}
-	virtual ~FBlockedPlayersListEOS() = default;
-};
-
-typedef TSharedRef<FBlockedPlayersListEOS> FBlockedPlayersListEOSRef;
-
-class FRecentPlayersListEOS :
-	public TOnlinePlayerList<FOnlineRecentPlayerEOSRef, FOnlineRecentPlayerEOSPtr>
-{
-public:
-	FRecentPlayersListEOS(int32 InLocalUserNum, FUniqueNetIdEOSRef InOwningNetId)
-		: TOnlinePlayerList<FOnlineRecentPlayerEOSRef, FOnlineRecentPlayerEOSPtr>(InLocalUserNum, InOwningNetId)
-	{
-	}
-	virtual ~FRecentPlayersListEOS() = default;
-};
-
-typedef TSharedRef<FRecentPlayersListEOS> FRecentPlayersListEOSRef;
-
 struct FNotificationIdCallbackPair
 {
 	EOS_NotificationId NotificationId;
@@ -435,12 +375,6 @@ private:
 	/** Per user friends lists accessible by user num or net id */
 	TMap<int32, FFriendsListEOSRef> LocalUserNumToFriendsListMap;
 	TMap<FString, FFriendsListEOSRef> NetIdStringToFriendsListMap;
-	/** Per user blocked player lists accessible by user num or net id */
-	TMap<int32, FBlockedPlayersListEOSRef> LocalUserNumToBlockedPlayerListMap;
-	TMap<FString, FBlockedPlayersListEOSRef> NetIdStringToBlockedPlayerListMap;
-	/** Per user recent player lists accessible by user num or net id */
-	TMap<int32, FRecentPlayersListEOSRef> LocalUserNumToRecentPlayerListMap;
-	TMap<FString, FRecentPlayersListEOSRef> NetIdStringToRecentPlayerListMap;
 
 	/** Ids mapped to remote users */
 	TMap<FString, FOnlineUserPtr> NetIdStringToOnlineUserMap;
