Added GetVoiceChatRoomName method to sessions interface, with EOS and EOSPlus implementations

diff --git a/Engine/Plugins/Online/OnlineSubsystem/Source/Private/OnlineSessionInterface.cpp b/Engine/Plugins/Online/OnlineSubsystem/Source/Private/OnlineSessionInterface.cpp
--- a/Engine/Plugins/Online/OnlineSubsystem/Source/Private/OnlineSessionInterface.cpp
+++ b/Engine/Plugins/Online/OnlineSubsystem/Source/Private/OnlineSessionInterface.cpp
@@ -49,6 +49,13 @@ bool IOnlineSession::StartMatchmaking(const TArray<FSessionMatchmakingUser>& Loc
 	return bStartMatchmakingSuccess;
 }
 
+FString IOnlineSession::GetVoiceChatRoomName(int32 LocalUserNum, const FName& SessionName)
+{
+	UE_LOG_ONLINE_SESSION(Warning, TEXT("[IOnlineSession::GetVoiceChatRoomName] This functionality is not implemented on this OSS."));
+
+	return FString();
+}
+
 void IOnlineSession::RemovePlayerFromSession(int32 LocalUserNum, FName SessionName, const FUniqueNetId& TargetPlayerId)
 {
 	UE_LOG_ONLINE_SESSION(Warning, TEXT("[IOnlineSession::RemovePlayerFromSession] This functionality is not implemented by default."));
diff --git a/Engine/Plugins/Online/OnlineSubsystem/Source/Public/Interfaces/OnlineSessionInterface.h b/Engine/Plugins/Online/OnlineSubsystem/Source/Public/Interfaces/OnlineSessionInterface.h
--- a/Engine/Plugins/Online/OnlineSubsystem/Source/Public/Interfaces/OnlineSessionInterface.h
+++ b/Engine/Plugins/Online/OnlineSubsystem/Source/Public/Interfaces/OnlineSessionInterface.h
@@ -846,6 +846,16 @@ public:
 	 */
 	virtual FOnlineSessionSettings* GetSessionSettings(FName SessionName) = 0;
 
+	/**
+	 * Returns the name for the Voice Chat room, if any were created by setting bUseLobbiesVoiceChatIfAvailable to true during session creation
+	 * 
+	 * @param LocalUserNum the user asking for the room name
+	 * @param SessionName the name of the related session
+	 * 
+	 * @return the name for the corresponding Voice Chat room, or an empty FString if there is none
+	 */
+	virtual FString ONLINESUBSYSTEM_API GetVoiceChatRoomName(int32 LocalUserNum, const FName& SessionName);
+
 	/**
 	 * Registers a player with the online service as being part of the online session
 	 *
diff --git a/Engine/Plugins/Online/OnlineSubsystemEOS/Source/OnlineSubsystemEOS/Private/OnlineSessionEOS.cpp b/Engine/Plugins/Online/OnlineSubsystemEOS/Source/OnlineSubsystemEOS/Private/OnlineSessionEOS.cpp
--- a/Engine/Plugins/Online/OnlineSubsystemEOS/Source/OnlineSubsystemEOS/Private/OnlineSessionEOS.cpp
+++ b/Engine/Plugins/Online/OnlineSubsystemEOS/Source/OnlineSubsystemEOS/Private/OnlineSessionEOS.cpp
@@ -3503,6 +3503,41 @@ FUniqueNetIdPtr FOnlineSessionEOS::CreateSessionIdFromString(const FString& Sess
 	return FUniqueNetIdEOSSession::Create(SessionIdStr);
 }
 
+FString FOnlineSessionEOS::GetVoiceChatRoomName(int32 LocalUserNum, const FName& SessionName)
+{
+	FString RTCRoomNameStr;
+
+	FNamedOnlineSession* Session = GetNamedSession(SessionName);
+	if (Session == nullptr)
+	{
+		UE_LOG_ONLINE_SESSION(Verbose, TEXT("[FOnlineSessionEOS::GetRTCRoomNameFromLobbyId] Unable to find session with name %s"), *SessionName.ToString());
+		return RTCRoomNameStr;
+	}
+
+	const FUniqueNetIdEOSLobby& LobbyId = FUniqueNetIdEOSLobby::Cast(Session->SessionInfo->GetSessionId());
+	const auto LobbyIdUTF8 = StringCast<UTF8CHAR>(*LobbyId.ToString());	
+
+	EOS_Lobby_GetRTCRoomNameOptions GetRTCRoomNameOptions = {};
+	GetRTCRoomNameOptions.ApiVersion = EOS_LOBBY_GETRTCROOMNAME_API_LATEST;
+	static_assert(EOS_LOBBY_GETRTCROOMNAME_API_LATEST == 1, "EOS_Lobby_GetRTCRoomNameOptions updated, check new fields");
+	GetRTCRoomNameOptions.LobbyId = (const char*)LobbyIdUTF8.Get();
+	GetRTCRoomNameOptions.LocalUserId = EOSSubsystem->UserManager->GetLocalProductUserId(LocalUserNum);
+
+	char RTCRoomNameUTF8[256];
+	uint32_t RTCRoomNameUTF8Length = UE_ARRAY_COUNT(RTCRoomNameUTF8);
+	const EOS_EResult Result = EOS_Lobby_GetRTCRoomName(LobbyHandle, &GetRTCRoomNameOptions, RTCRoomNameUTF8, &RTCRoomNameUTF8Length);
+	if (Result == EOS_EResult::EOS_Success)
+	{
+		RTCRoomNameStr = UTF8_TO_TCHAR(RTCRoomNameUTF8);
+	}
+	else
+	{
+		UE_LOG_ONLINE_SESSION(Verbose, TEXT("[FOnlineSessionEOS::GetRTCRoomNameFromLobbyId] EOS_Lobby_GetRTCRoomName not successful. Finished with EOS_EResult %s"), *LexToString(Result));
+	}
+
+	return RTCRoomNameStr;
+}
+
 EOS_ELobbyPermissionLevel FOnlineSessionEOS::GetLobbyPermissionLevelFromSessionSettings(const FOnlineSessionSettings& SessionSettings)
 {
 	EOS_ELobbyPermissionLevel Result;
diff --git a/Engine/Plugins/Online/OnlineSubsystemEOS/Source/OnlineSubsystemEOS/Private/OnlineSessionEOS.h b/Engine/Plugins/Online/OnlineSubsystemEOS/Source/OnlineSubsystemEOS/Private/OnlineSessionEOS.h
--- a/Engine/Plugins/Online/OnlineSubsystemEOS/Source/OnlineSubsystemEOS/Private/OnlineSessionEOS.h
+++ b/Engine/Plugins/Online/OnlineSubsystemEOS/Source/OnlineSubsystemEOS/Private/OnlineSessionEOS.h
@@ -144,6 +144,7 @@ public:
 	virtual bool GetResolvedConnectString(FName SessionName, FString& ConnectInfo, FName PortType) override;
 	virtual bool GetResolvedConnectString(const FOnlineSessionSearchResult& SearchResult, FName PortType, FString& ConnectInfo) override;
 	virtual FOnlineSessionSettings* GetSessionSettings(FName SessionName) override;
+	virtual FString GetVoiceChatRoomName(int32 LocalUserNum, const FName& SessionName) override;
 	virtual bool RegisterPlayer(FName SessionName, const FUniqueNetId& PlayerId, bool bWasInvited) override;
 	virtual bool RegisterPlayers(FName SessionName, const TArray< FUniqueNetIdRef >& Players, bool bWasInvited = false) override;
 	virtual bool UnregisterPlayer(FName SessionName, const FUniqueNetId& PlayerId) override;
diff --git a/Engine/Plugins/Online/OnlineSubsystemEOS/Source/OnlineSubsystemEOSPlus/Private/OnlineSessionEOSPlus.cpp b/Engine/Plugins/Online/OnlineSubsystemEOS/Source/OnlineSubsystemEOSPlus/Private/OnlineSessionEOSPlus.cpp
--- a/Engine/Plugins/Online/OnlineSubsystemEOS/Source/OnlineSubsystemEOSPlus/Private/OnlineSessionEOSPlus.cpp
+++ b/Engine/Plugins/Online/OnlineSubsystemEOS/Source/OnlineSubsystemEOSPlus/Private/OnlineSessionEOSPlus.cpp
@@ -808,6 +808,16 @@ FOnlineSessionSettings* FOnlineSessionEOSPlus::GetSessionSettings(FName SessionN
 	return BaseSessionInterface->GetSessionSettings(SessionName);
 }
 
+FString FOnlineSessionEOSPlus::GetVoiceChatRoomName(int32 LocalUserNum, const FName& SessionName)
+{
+	if (bUseEOSSessions)
+	{
+		return EOSSessionInterface->GetVoiceChatRoomName(LocalUserNum, SessionName);
+	}
+
+	return BaseSessionInterface->GetVoiceChatRoomName(LocalUserNum, SessionName);
+}
+
 bool FOnlineSessionEOSPlus::RegisterPlayer(FName SessionName, const FUniqueNetId& PlayerId, bool bWasInvited)
 {
 	if (bUseEOSSessions)
diff --git a/Engine/Plugins/Online/OnlineSubsystemEOS/Source/OnlineSubsystemEOSPlus/Private/OnlineSessionEOSPlus.h b/Engine/Plugins/Online/OnlineSubsystemEOS/Source/OnlineSubsystemEOSPlus/Private/OnlineSessionEOSPlus.h
--- a/Engine/Plugins/Online/OnlineSubsystemEOS/Source/OnlineSubsystemEOSPlus/Private/OnlineSessionEOSPlus.h
+++ b/Engine/Plugins/Online/OnlineSubsystemEOS/Source/OnlineSubsystemEOSPlus/Private/OnlineSessionEOSPlus.h
@@ -49,6 +49,7 @@ public:
 	virtual bool GetResolvedConnectString(FName SessionName, FString& ConnectInfo, FName PortType) override;
 	virtual bool GetResolvedConnectString(const FOnlineSessionSearchResult& SearchResult, FName PortType, FString& ConnectInfo) override;
 	virtual FOnlineSessionSettings* GetSessionSettings(FName SessionName) override;
+	virtual FString GetVoiceChatRoomName(int32 LocalUserNum, const FName& SessionName) override;
 	virtual bool RegisterPlayer(FName SessionName, const FUniqueNetId& PlayerId, bool bWasInvited) override;
 	virtual bool RegisterPlayers(FName SessionName, const TArray< FUniqueNetIdRef >& Players, bool bWasInvited = false) override;
 	virtual bool UnregisterPlayer(FName SessionName, const FUniqueNetId& PlayerId) override;
