- Adjust bUseEOSConnect doc

diff --git a/Engine/Plugins/Online/OnlineSubsystemEOS/Source/OnlineSubsystemEOS/Public/EOSSettings.h b/Engine/Plugins/Online/OnlineSubsystemEOS/Source/OnlineSubsystemEOS/Public/EOSSettings.h
--- a/Engine/Plugins/Online/OnlineSubsystemEOS/Source/OnlineSubsystemEOS/Public/EOSSettings.h
+++ b/Engine/Plugins/Online/OnlineSubsystemEOS/Source/OnlineSubsystemEOS/Public/EOSSettings.h
@@ -144,7 +144,7 @@ public:
 	bool bUseEAS = false;
 
 	/** Set to true to have EOS Connect APIs used to link accounts for crossplay */
-	UPROPERTY(Config, EditAnywhere, BlueprintReadOnly, Category="EOSPlus Login Settings", DisplayName="Use Crossplatform User IDs for EOS Login (doesn't use Epic Account)")
+	UPROPERTY(Config, EditAnywhere, BlueprintReadOnly, Category="EOSPlus Login Settings", DisplayName="Use EOS Connect APIs to create and link Product User IDs (PUIDs), and use EOS Game Services")
 	bool bUseEOSConnect = false;
 
 	/** Set to true to write stats to EOS as well as the default platform */
