Fix for user UniqueId retrieval in QueryStats callback

diff --git a/Engine/Plugins/Online/OnlineSubsystemEOS/Source/OnlineSubsystemEOS/Private/OnlineStatsEOS.cpp b/Engine/Plugins/Online/OnlineSubsystemEOS/Source/OnlineSubsystemEOS/Private/OnlineStatsEOS.cpp
--- a/Engine/Plugins/Online/OnlineSubsystemEOS/Source/OnlineSubsystemEOS/Private/OnlineStatsEOS.cpp
+++ b/Engine/Plugins/Online/OnlineSubsystemEOS/Source/OnlineSubsystemEOS/Private/OnlineStatsEOS.cpp
@@ -158,42 +158,38 @@ void FOnlineStatsEOS::QueryStats(const FUniqueNetIdRef LocalUserId, const TArray
 		Options.TargetUserId = TargetEOSUserId;
 
 		FReadStatsCallback* CallbackObj = new FReadStatsCallback(FOnlineStatsEOSWeakPtr(AsShared()));
-		CallbackObj->CallbackLambda = [this, StatsQueryContext](const EOS_Stats_OnQueryStatsCompleteCallbackInfo* Data)
+		CallbackObj->CallbackLambda = [this, StatUserId, StatsQueryContext](const EOS_Stats_OnQueryStatsCompleteCallbackInfo* Data)
 		{
 			StatsQueryContext->NumPlayerReads--;
 			bool bWasSuccessful = Data->ResultCode == EOS_EResult::EOS_Success;
 			if (bWasSuccessful)
 			{
-				FUniqueNetIdEOSPtr StatUserId = EOSSubsystem->UserManager->GetLocalUniqueNetIdEOS(Data->TargetUserId);
-				if (StatUserId.IsValid())
+				char StatNameANSI[EOS_OSS_STRING_BUFFER_LENGTH];
+				EOS_Stats_CopyStatByNameOptions Options = { };
+				Options.ApiVersion = EOS_STATS_COPYSTATBYNAME_API_LATEST;
+				Options.TargetUserId = Data->TargetUserId;
+				Options.Name = StatNameANSI;
+
+				TSharedPtr<FOnlineStatsUserStats> UserStats = StatsCache.Emplace(StatUserId, MakeShared<FOnlineStatsUserStats>(StatUserId));
+				// Read each stat that we were looking for so we can mark missing ones as "empty"
+				for (const FString& StatName : StatsQueryContext->StatNames)
 				{
-					char StatNameANSI[EOS_OSS_STRING_BUFFER_LENGTH];
-					EOS_Stats_CopyStatByNameOptions Options = { };
-					Options.ApiVersion = EOS_STATS_COPYSTATBYNAME_API_LATEST;
-					Options.TargetUserId = Data->TargetUserId;
-					Options.Name = StatNameANSI;
-
-					TSharedPtr<FOnlineStatsUserStats> UserStats = StatsCache.Emplace(StatUserId.ToSharedRef(), MakeShared<FOnlineStatsUserStats>(StatUserId.ToSharedRef()));
-					// Read each stat that we were looking for so we can mark missing ones as "empty"
-					for (const FString& StatName : StatsQueryContext->StatNames)
+					FCStringAnsi::Strncpy(StatNameANSI, TCHAR_TO_UTF8(*StatName.ToUpper()), EOS_OSS_STRING_BUFFER_LENGTH);
+
+					EOS_Stats_Stat* ReadStat = nullptr;
+					if (EOS_Stats_CopyStatByName(EOSSubsystem->StatsHandle, &Options, &ReadStat) == EOS_EResult::EOS_Success)
+					{
+						UE_LOG_ONLINE_STATS(VeryVerbose, TEXT("Found value for stat %s"), *StatName);
+
+						UserStats->Stats.Add(StatName, FOnlineStatValue(ReadStat->Value));
+
+						EOS_Stats_Stat_Release(ReadStat);
+					}
+					else
 					{
-						FCStringAnsi::Strncpy(StatNameANSI, TCHAR_TO_UTF8(*StatName.ToUpper()), EOS_OSS_STRING_BUFFER_LENGTH);
-
-						EOS_Stats_Stat* ReadStat = nullptr;
-						if (EOS_Stats_CopyStatByName(EOSSubsystem->StatsHandle, &Options, &ReadStat) == EOS_EResult::EOS_Success)
-						{
-							UE_LOG_ONLINE_STATS(VeryVerbose, TEXT("Found value for stat %s"), *StatName);
-
-							UserStats->Stats.Add(StatName, FOnlineStatValue(ReadStat->Value));
-
-							EOS_Stats_Stat_Release(ReadStat);
-						}
-						else
-						{
-							// Put an empty stat in
-							UE_LOG_ONLINE_STATS(VeryVerbose, TEXT("Value not found for stat %s, adding empty value"), *StatName);
-							UserStats->Stats.Add(StatName, FOnlineStatValue());
-						}
+						// Put an empty stat in
+						UE_LOG_ONLINE_STATS(VeryVerbose, TEXT("Value not found for stat %s, adding empty value"), *StatName);
+						UserStats->Stats.Add(StatName, FOnlineStatValue());
 					}
 				}
 			}
