- Rename EncryptionKey to ClientEncryptionKey in config to fix EncryptionKey getting removed during packaging due to IniKeyDenylist=EncryptionKey entry in BaseEngine.ini, breaking Player/TitleDataStorage

diff --git a/Engine/Config/ConfigRedirects.ini b/Engine/Config/ConfigRedirects.ini
--- a/Engine/Config/ConfigRedirects.ini
+++ b/Engine/Config/ConfigRedirects.ini
@@ -38,4 +38,13 @@ RPCBlockWhitelist=RPCBlockAllowlist
 bDisableDomainWhitelist=bDisableDomainAllowlist
 
 [WinHttpWebSocket]
-bDisableDomainWhitelist=bDisableDomainAllowlist
\ No newline at end of file
+bDisableDomainWhitelist=bDisableDomainAllowlist
+
+[OnlineServices.EOS]
+EncryptionKey=ClientEncryptionKey
+
+[EOSVoiceChat]
+EncryptionKey=ClientEncryptionKey
+
+[/Script/OnlineSubsystemEOS.EOSSettings]
+EncryptionKey=ClientEncryptionKey
\ No newline at end of file
diff --git a/Engine/Plugins/Online/EOSShared/Source/EOSShared/Private/EOSSDKManager.cpp b/Engine/Plugins/Online/EOSShared/Source/EOSShared/Private/EOSSDKManager.cpp
--- a/Engine/Plugins/Online/EOSShared/Source/EOSShared/Private/EOSSDKManager.cpp
+++ b/Engine/Plugins/Online/EOSShared/Source/EOSShared/Private/EOSSDKManager.cpp
@@ -284,7 +284,13 @@ const FEOSSDKPlatformConfig* FEOSSDKManager::GetPlatformConfig(const FString& Pl
 	GConfig->GetString(*SectionName, TEXT("SandboxId"), PlatformConfig->SandboxId, GEngineIni);
 	GConfig->GetString(*SectionName, TEXT("ClientId"), PlatformConfig->ClientId, GEngineIni);
 	GConfig->GetString(*SectionName, TEXT("ClientSecret"), PlatformConfig->ClientSecret, GEngineIni);
-	GConfig->GetString(*SectionName, TEXT("EncryptionKey"), PlatformConfig->EncryptionKey, GEngineIni);
+	if (GConfig->GetString(*SectionName, TEXT("EncryptionKey"), PlatformConfig->EncryptionKey, GEngineIni))
+	{
+		// EncryptionKey gets removed from packaged builds due to IniKeyDenylist=EncryptionKey entry in BaseGame.ini.
+		// Normally we could just add a remap in ConfigRedirects.ini but the section name varies with the PlatformConfigName.
+		UE_LOG(LogEOSSDK, Warning, TEXT("Config section \"EOSSDK.Platform.%s\" contains deprecated key EncryptionKey, please migrate to ClientEncryptionKey."), *PlatformConfigName);
+	}
+	GConfig->GetString(*SectionName, TEXT("ClientEncryptionKey"), PlatformConfig->EncryptionKey, GEngineIni);
 	GConfig->GetString(*SectionName, TEXT("OverrideCountryCode"), PlatformConfig->OverrideCountryCode, GEngineIni);
 	GConfig->GetString(*SectionName, TEXT("OverrideLocaleCode"), PlatformConfig->OverrideLocaleCode, GEngineIni);
 	GConfig->GetString(*SectionName, TEXT("DeploymentId"), PlatformConfig->DeploymentId, GEngineIni);
diff --git a/Engine/Plugins/Online/OnlineServicesEOSGS/Source/Private/Online/OnlineServicesEOSGSPlatformFactory.cpp b/Engine/Plugins/Online/OnlineServicesEOSGS/Source/Private/Online/OnlineServicesEOSGSPlatformFactory.cpp
--- a/Engine/Plugins/Online/OnlineServicesEOSGS/Source/Private/Online/OnlineServicesEOSGSPlatformFactory.cpp
+++ b/Engine/Plugins/Online/OnlineServicesEOSGS/Source/Private/Online/OnlineServicesEOSGSPlatformFactory.cpp
@@ -33,7 +33,8 @@ FEOSPlatformConfig LoadEOSPlatformConfig()
 	GConfig->GetString(*ConfigSection, TEXT("DeploymentId"), PlatformConfig.DeploymentId, GEngineIni);
 	GConfig->GetString(*ConfigSection, TEXT("ClientId"), PlatformConfig.ClientId, GEngineIni);
 	GConfig->GetString(*ConfigSection, TEXT("ClientSecret"), PlatformConfig.ClientSecret, GEngineIni);
-	GConfig->GetString(*ConfigSection, TEXT("EncryptionKey"), PlatformConfig.EncryptionKey, GEngineIni);
+	// Config key renamed to ClientEncryptionKey as EncryptionKey gets removed from packaged builds due to IniKeyDenylist=EncryptionKey entry in BaseGame.ini.
+	GConfig->GetString(*ConfigSection, TEXT("ClientEncryptionKey"), PlatformConfig.EncryptionKey, GEngineIni);
 	return PlatformConfig;
 }
 
diff --git a/Engine/Plugins/Online/OnlineSubsystemEOS/Config/BaseOnlineSubsystemEOS.ini b/Engine/Plugins/Online/OnlineSubsystemEOS/Config/BaseOnlineSubsystemEOS.ini
new file mode 100644
--- /dev/null
+++ b/Engine/Plugins/Online/OnlineSubsystemEOS/Config/BaseOnlineSubsystemEOS.ini
@@ -0,0 +1,3 @@
+[CoreRedirects]
+	
++PropertyRedirects=(OldName="ArtifactSettings.EncryptionKey", NewName="ArtifactSettings.ClientEncryptionKey")
diff --git a/Engine/Plugins/Online/OnlineSubsystemEOS/Source/OnlineSubsystemEOS/Private/EOSSettings.cpp b/Engine/Plugins/Online/OnlineSubsystemEOS/Source/OnlineSubsystemEOS/Private/EOSSettings.cpp
--- a/Engine/Plugins/Online/OnlineSubsystemEOS/Source/OnlineSubsystemEOS/Private/EOSSettings.cpp
+++ b/Engine/Plugins/Online/OnlineSubsystemEOS/Source/OnlineSubsystemEOS/Private/EOSSettings.cpp
@@ -63,7 +63,7 @@ FEOSArtifactSettings FArtifactSettings::ToNative() const
 	Native.ClientId = ClientId;
 	Native.ClientSecret = ClientSecret;
 	Native.DeploymentId = DeploymentId;
-	Native.EncryptionKey = EncryptionKey;
+	Native.EncryptionKey = ClientEncryptionKey;
 	Native.ProductId = ProductId;
 	Native.SandboxId = SandboxId;
 
@@ -118,7 +118,7 @@ FEOSArtifactSettings ParseArtifactSettingsFromConfigString(const FString& RawLin
 		{
 			Result.DeploymentId = StripQuotes(Values[ValueIndex + 1]);
 		}
-		else if (Values[ValueIndex] == TEXT("EncryptionKey"))
+		else if (Values[ValueIndex] == TEXT("ClientEncryptionKey"))
 		{
 			Result.EncryptionKey = StripQuotes(Values[ValueIndex + 1]);
 		}
@@ -353,13 +353,13 @@ void UEOSSettings::PostEditChangeProperty(FPropertyChangedEvent& PropertyChanged
 				}
 			}
 
-			if (!Artifact.EncryptionKey.IsEmpty())
+			if (!Artifact.ClientEncryptionKey.IsEmpty())
 			{
-				if (!IsHex(Artifact.EncryptionKey) || Artifact.EncryptionKey.Len() != 64)
+				if (!IsHex(Artifact.ClientEncryptionKey) || Artifact.ClientEncryptionKey.Len() != 64)
 				{
 					FMessageDialog::Open(EAppMsgType::Ok,
-						LOCTEXT("EncryptionKeyNotHexMsg", "EncryptionKey must contain 64 hex characters"));
-					Artifact.EncryptionKey.Empty();
+						LOCTEXT("EncryptionKeyNotHexMsg", "ClientEncryptionKey must contain 64 hex characters"));
+					Artifact.ClientEncryptionKey.Empty();
 				}
 			}
 		}
diff --git a/Engine/Plugins/Online/OnlineSubsystemEOS/Source/OnlineSubsystemEOS/Public/EOSSettings.h b/Engine/Plugins/Online/OnlineSubsystemEOS/Source/OnlineSubsystemEOS/Public/EOSSettings.h
--- a/Engine/Plugins/Online/OnlineSubsystemEOS/Source/OnlineSubsystemEOS/Public/EOSSettings.h
+++ b/Engine/Plugins/Online/OnlineSubsystemEOS/Source/OnlineSubsystemEOS/Public/EOSSettings.h
@@ -57,7 +57,8 @@ public:
 	FString DeploymentId;
 
 	UPROPERTY(Config, EditAnywhere, BlueprintReadOnly, Category="EOS Artifact Settings")
-	FString EncryptionKey;
+	// Config key renamed to ClientEncryptionKey as EncryptionKey gets removed from packaged builds due to IniKeyDenylist=EncryptionKey entry in BaseGame.ini.
+	FString ClientEncryptionKey;
 
 	FEOSArtifactSettings ToNative() const;
 };
diff --git a/Engine/Plugins/Online/VoiceChat/EOSVoiceChat/Source/EOSVoiceChat/Private/EOSVoiceChat.cpp b/Engine/Plugins/Online/VoiceChat/EOSVoiceChat/Source/EOSVoiceChat/Private/EOSVoiceChat.cpp
--- a/Engine/Plugins/Online/VoiceChat/EOSVoiceChat/Source/EOSVoiceChat/Private/EOSVoiceChat.cpp
+++ b/Engine/Plugins/Online/VoiceChat/EOSVoiceChat/Source/EOSVoiceChat/Private/EOSVoiceChat.cpp
@@ -120,7 +120,7 @@ void FEOSVoiceChat::Initialize(const FOnVoiceChatInitializeCompleteDelegate& Ini
 					GConfig->GetString(TEXT("EOSVoiceChat"), TEXT("DeploymentId"), ConfigDeploymentId, GEngineIni);
 					GConfig->GetString(TEXT("EOSVoiceChat"), TEXT("ClientId"), ConfigClientId, GEngineIni);
 					GConfig->GetString(TEXT("EOSVoiceChat"), TEXT("ClientSecret"), ConfigClientSecret, GEngineIni);
-					GConfig->GetString(TEXT("EOSVoiceChat"), TEXT("EncryptionKey"), ConfigEncryptionKey, GEngineIni);
+					GConfig->GetString(TEXT("EOSVoiceChat"), TEXT("ClientEncryptionKey"), ConfigEncryptionKey, GEngineIni);
 					GConfig->GetString(TEXT("EOSVoiceChat"), TEXT("OverrideCountryCode"), ConfigOverrideCountryCode, GEngineIni);
 					GConfig->GetString(TEXT("EOSVoiceChat"), TEXT("OverrideLocaleCode"), ConfigOverrideLocaleCode, GEngineIni);
 
