- Set specific ApiVersion when calling EOS api's, not API_LATEST, so it's a conscious decision what API version we are targeting.

diff --git a/Engine/Plugins/Online/EOSShared/Source/EOSShared/EOSShared.Build.cs b/Engine/Plugins/Online/EOSShared/Source/EOSShared/EOSShared.Build.cs
--- a/Engine/Plugins/Online/EOSShared/Source/EOSShared/EOSShared.Build.cs
+++ b/Engine/Plugins/Online/EOSShared/Source/EOSShared/EOSShared.Build.cs
@@ -1,15 +1,33 @@
 // Copyright Epic Games, Inc. All Rights Reserved.
 
+using EpicGames.Core;
 using System;
 using System.IO;
 using UnrealBuildTool;
 
 public class EOSShared : ModuleRules
 {
+	bool bEnableApiVersionWarnings
+	{
+		get
+		{
+			bool bResult;
+			ConfigHierarchy Config = ConfigCache.ReadHierarchy(ConfigHierarchyType.Engine, DirectoryReference.FromFile(Target.ProjectFile), Target.Platform);
+			if(!Config.GetBool("EOSShared", "bEnableApiVersionWarnings", out bResult))
+			{
+				// Default to true
+				bResult = true;
+			}
+			return bResult;
+		}
+	}
+
 	public EOSShared(ReadOnlyTargetRules Target) : base(Target)
 	{
 		Type = ModuleType.CPlusPlus;
 
+		PublicDefinitions.Add("UE_WITH_EOS_SDK_APIVERSION_WARNINGS=" + (bEnableApiVersionWarnings ? "1" : "0"));
+
 		PrivateDependencyModuleNames.AddRange(
 			new string[]
 			{
diff --git a/Engine/Plugins/Online/EOSShared/Source/EOSShared/Private/Android/AndroidEOSSDKManager.cpp b/Engine/Plugins/Online/EOSShared/Source/EOSShared/Private/Android/AndroidEOSSDKManager.cpp
--- a/Engine/Plugins/Online/EOSShared/Source/EOSShared/Private/Android/AndroidEOSSDKManager.cpp
+++ b/Engine/Plugins/Online/EOSShared/Source/EOSShared/Private/Android/AndroidEOSSDKManager.cpp
@@ -1,17 +1,19 @@
 // Copyright Epic Games, Inc. All Rights Reserved.
 
 #include "AndroidEOSSDKManager.h"
-#include "HAL/FileManager.h"
 
 #if WITH_EOS_SDK
 
+#include "HAL/FileManager.h"
+#include "EOSShared.h"
+
 #include "eos_android.h"
 
 EOS_EResult FAndroidEOSSDKManager::EOSInitialize(EOS_InitializeOptions& Options)
 {
 	EOS_Android_InitializeOptions SystemInitializeOptions = { 0 };
-	SystemInitializeOptions.ApiVersion = EOS_ANDROID_INITIALIZEOPTIONS_API_LATEST;
-	static_assert(EOS_ANDROID_INITIALIZEOPTIONS_API_LATEST == 2, "EOS_Android_InitializeOptions updated, check new fields");
+	SystemInitializeOptions.ApiVersion = 2;
+	UE_EOS_CHECK_API_MISMATCH(EOS_ANDROID_INITIALIZEOPTIONS_API_LATEST, 2);
 	SystemInitializeOptions.Reserved = nullptr;
 	SystemInitializeOptions.OptionalInternalDirectory = nullptr;
 	SystemInitializeOptions.OptionalExternalDirectory = nullptr;
diff --git a/Engine/Plugins/Online/EOSShared/Source/EOSShared/Private/EOSSDKManager.cpp b/Engine/Plugins/Online/EOSShared/Source/EOSShared/Private/EOSSDKManager.cpp
--- a/Engine/Plugins/Online/EOSShared/Source/EOSShared/Private/EOSSDKManager.cpp
+++ b/Engine/Plugins/Online/EOSShared/Source/EOSShared/Private/EOSSDKManager.cpp
@@ -209,8 +209,8 @@ EOS_EResult FEOSSDKManager::Initialize()
 		const FTCHARToUTF8 ProductVersion(*GetProductVersion());
 
 		EOS_InitializeOptions InitializeOptions = {};
-		InitializeOptions.ApiVersion = EOS_INITIALIZE_API_LATEST;
-		static_assert(EOS_INITIALIZE_API_LATEST == 4, "EOS_InitializeOptions updated, check new fields");
+		InitializeOptions.ApiVersion = 4;
+		UE_EOS_CHECK_API_MISMATCH(EOS_INITIALIZE_API_LATEST, 4);
 		InitializeOptions.AllocateMemoryFunction = &EosMalloc;
 		InitializeOptions.ReallocateMemoryFunction = &EosRealloc;
 		InitializeOptions.ReleaseMemoryFunction = &EosFree;
@@ -398,9 +398,9 @@ IEOSPlatformHandlePtr FEOSSDKManager::CreatePlatform(const FString& PlatformConf
 	const FTCHARToUTF8 Utf8DeploymentId(*PlatformConfig->DeploymentId);
 	const FTCHARToUTF8 Utf8CacheDirectory(*PlatformConfig->CacheDirectory);
 
-	static_assert(EOS_PLATFORM_OPTIONS_API_LATEST == 12, "EOS_Platform_Options updated");
 	EOS_Platform_Options PlatformOptions = {};
-	PlatformOptions.ApiVersion = EOS_PLATFORM_OPTIONS_API_LATEST;
+	PlatformOptions.ApiVersion = 12;
+	UE_EOS_CHECK_API_MISMATCH(EOS_PLATFORM_OPTIONS_API_LATEST, 12);
 	PlatformOptions.Reserved = nullptr;
 	PlatformOptions.ProductId = Utf8ProductId.Length() ? Utf8ProductId.Get() : nullptr;
 	PlatformOptions.SandboxId = Utf8SandboxId.Length() ? Utf8SandboxId.Get() : nullptr;
@@ -428,9 +428,9 @@ IEOSPlatformHandlePtr FEOSSDKManager::CreatePlatform(const FString& PlatformConf
 
 	PlatformOptions.TickBudgetInMilliseconds = PlatformConfig->TickBudgetInMilliseconds;
 
-	static_assert(EOS_PLATFORM_RTCOPTIONS_API_LATEST == 1, "EOS_Platform_RTCOptions updated");
 	EOS_Platform_RTCOptions PlatformRTCOptions = {};
-	PlatformRTCOptions.ApiVersion = EOS_PLATFORM_RTCOPTIONS_API_LATEST;
+	PlatformRTCOptions.ApiVersion = 1;
+	UE_EOS_CHECK_API_MISMATCH(EOS_PLATFORM_RTCOPTIONS_API_LATEST, 1);
 	PlatformRTCOptions.PlatformSpecificOptions = nullptr;
 	PlatformOptions.RTCOptions = PlatformConfig->bEnableRTC ? &PlatformRTCOptions : nullptr;
 
@@ -755,9 +755,9 @@ void FEOSSDKManager::LogPlatformInfo(const EOS_HPlatform Platform, int32 Indent)
 	UE_LOG_EOSSDK_INFO("OverrideCountryCode=%s", *GetOverrideCountryCode(Platform));
 	UE_LOG_EOSSDK_INFO("OverrideLocaleCode=%s", *GetOverrideLocaleCode(Platform));
 
-	static_assert(EOS_PLATFORM_GETDESKTOPCROSSPLAYSTATUS_API_LATEST == 1, "EOS_Platform_GetDesktopCrossplayStatusOptions updated");
 	EOS_Platform_GetDesktopCrossplayStatusOptions GetDesktopCrossplayStatusOptions;
-	GetDesktopCrossplayStatusOptions.ApiVersion = EOS_PLATFORM_GETDESKTOPCROSSPLAYSTATUS_API_LATEST;
+	GetDesktopCrossplayStatusOptions.ApiVersion = 1;
+	UE_EOS_CHECK_API_MISMATCH(EOS_PLATFORM_GETDESKTOPCROSSPLAYSTATUS_API_LATEST, 1);
 
 	EOS_Platform_GetDesktopCrossplayStatusInfo GetDesktopCrossplayStatusInfo;
 	EOS_EResult Result = EOS_Platform_GetDesktopCrossplayStatus(Platform, &GetDesktopCrossplayStatusOptions, &GetDesktopCrossplayStatusInfo);
@@ -807,9 +807,9 @@ void FEOSSDKManager::LogAuthInfo(const EOS_HPlatform Platform, const EOS_EpicAcc
 	const EOS_HAuth AuthHandle = EOS_Platform_GetAuthInterface(Platform);
 	UE_LOG_EOSSDK_INFO("LoginStatus=%s", LexToString(EOS_Auth_GetLoginStatus(AuthHandle, LoggedInAccount)));
 
-	static_assert(EOS_AUTH_COPYUSERAUTHTOKEN_API_LATEST == 1, "EOS_Auth_CopyUserAuthTokenOptions updated");
 	EOS_Auth_CopyUserAuthTokenOptions CopyUserAuthTokenOptions;
-	CopyUserAuthTokenOptions.ApiVersion = EOS_AUTH_COPYUSERAUTHTOKEN_API_LATEST;
+	CopyUserAuthTokenOptions.ApiVersion = 1;
+	UE_EOS_CHECK_API_MISMATCH(EOS_AUTH_COPYUSERAUTHTOKEN_API_LATEST, 1);
 
 	EOS_Auth_Token* AuthToken;
 	EOS_EResult Result = EOS_Auth_CopyUserAuthToken(AuthHandle, &CopyUserAuthTokenOptions, LoggedInAccount, &AuthToken);
@@ -839,9 +839,9 @@ void FEOSSDKManager::LogAuthInfo(const EOS_HPlatform Platform, const EOS_EpicAcc
 	}
 
 #if !UE_BUILD_SHIPPING
-	static_assert(EOS_AUTH_COPYIDTOKEN_API_LATEST == 1, "EOS_Auth_CopyIdTokenOptions updated");
 	EOS_Auth_CopyIdTokenOptions CopyIdTokenOptions;
-	CopyIdTokenOptions.ApiVersion = EOS_AUTH_COPYIDTOKEN_API_LATEST;
+	CopyIdTokenOptions.ApiVersion = 1;
+	UE_EOS_CHECK_API_MISMATCH(EOS_AUTH_COPYIDTOKEN_API_LATEST, 1);
 	CopyIdTokenOptions.AccountId = LoggedInAccount;
 
 	EOS_Auth_IdToken* IdToken;
@@ -862,9 +862,9 @@ void FEOSSDKManager::LogUserInfo(const EOS_HPlatform Platform, const EOS_EpicAcc
 {
 	UE_LOG_EOSSDK_INFO("EpicAccountId=%s", *LexToString(TargetAccount));
 
-	static_assert(EOS_USERINFO_COPYUSERINFO_API_LATEST >=2 && EOS_USERINFO_COPYUSERINFO_API_LATEST <= 3, "EOS_UserInfo_CopyUserInfo updated");
 	EOS_UserInfo_CopyUserInfoOptions CopyUserInfoOptions;
-	CopyUserInfoOptions.ApiVersion = EOS_USERINFO_COPYUSERINFO_API_LATEST;
+	CopyUserInfoOptions.ApiVersion = 3;
+	UE_EOS_CHECK_API_MISMATCH(EOS_USERINFO_COPYUSERINFO_API_LATEST, 3);
 	CopyUserInfoOptions.LocalUserId = LoggedInAccount;
 	CopyUserInfoOptions.TargetUserId = TargetAccount;
 
@@ -892,9 +892,9 @@ void FEOSSDKManager::LogUserInfo(const EOS_HPlatform Platform, const EOS_EpicAcc
 
 void FEOSSDKManager::LogPresenceInfo(const EOS_HPlatform Platform, const EOS_EpicAccountId LoggedInAccount, const EOS_EpicAccountId TargetAccount, int32 Indent) const
 {
-	static_assert(EOS_PRESENCE_HASPRESENCE_API_LATEST == 1, "EOS_Presence_HasPresenceOptions updated");
 	EOS_Presence_HasPresenceOptions HasPresenceOptions;
-	HasPresenceOptions.ApiVersion = EOS_PRESENCE_HASPRESENCE_API_LATEST;
+	HasPresenceOptions.ApiVersion = 1;
+	UE_EOS_CHECK_API_MISMATCH(EOS_PRESENCE_HASPRESENCE_API_LATEST, 1);
 	HasPresenceOptions.LocalUserId = LoggedInAccount;
 	HasPresenceOptions.TargetUserId = TargetAccount;
 
@@ -905,9 +905,9 @@ void FEOSSDKManager::LogPresenceInfo(const EOS_HPlatform Platform, const EOS_Epi
 		return;
 	}
 
-	static_assert(EOS_PRESENCE_COPYPRESENCE_API_LATEST >=2 && EOS_PRESENCE_COPYPRESENCE_API_LATEST <= 3, "EOS_Presence_CopyPresenceOptions updated");
 	EOS_Presence_CopyPresenceOptions CopyPresenceOptions;
-	CopyPresenceOptions.ApiVersion = EOS_PRESENCE_COPYPRESENCE_API_LATEST;
+	CopyPresenceOptions.ApiVersion = 3;
+	UE_EOS_CHECK_API_MISMATCH(EOS_PRESENCE_COPYPRESENCE_API_LATEST, 3);
 	CopyPresenceOptions.LocalUserId = LoggedInAccount;
 	CopyPresenceOptions.TargetUserId = TargetAccount;
 
@@ -943,18 +943,18 @@ void FEOSSDKManager::LogPresenceInfo(const EOS_HPlatform Platform, const EOS_Epi
 
 void FEOSSDKManager::LogFriendsInfo(const EOS_HPlatform Platform, const EOS_EpicAccountId LoggedInAccount, int32 Indent) const
 {
-	static_assert(EOS_FRIENDS_GETFRIENDSCOUNT_API_LATEST == 1, "EOS_Friends_GetFriendsCountOptions updated");
 	EOS_Friends_GetFriendsCountOptions FriendsCountOptions;
-	FriendsCountOptions.ApiVersion = EOS_FRIENDS_GETFRIENDSCOUNT_API_LATEST;
+	FriendsCountOptions.ApiVersion = 1;
+	UE_EOS_CHECK_API_MISMATCH(EOS_FRIENDS_GETFRIENDSCOUNT_API_LATEST, 1);
 	FriendsCountOptions.LocalUserId = LoggedInAccount;
 
 	const EOS_HFriends FriendsHandle = EOS_Platform_GetFriendsInterface(Platform);
 	const int32_t FriendsCount = EOS_Friends_GetFriendsCount(FriendsHandle, &FriendsCountOptions);
 	UE_LOG_EOSSDK_INFO("Friends=%d", FriendsCount);
 
-	static_assert(EOS_FRIENDS_GETFRIENDATINDEX_API_LATEST == 1, "EOS_Friends_GetFriendAtIndexOptions updated");
 	EOS_Friends_GetFriendAtIndexOptions FriendAtIndexOptions;
-	FriendAtIndexOptions.ApiVersion = EOS_FRIENDS_GETFRIENDATINDEX_API_LATEST;
+	FriendAtIndexOptions.ApiVersion = 1;
+	UE_EOS_CHECK_API_MISMATCH(EOS_FRIENDS_GETFRIENDATINDEX_API_LATEST, 1);
 	FriendAtIndexOptions.LocalUserId = LoggedInAccount;
 
 	for (int32_t FriendIndex = 0; FriendIndex < FriendsCount; FriendIndex++)
@@ -964,9 +964,9 @@ void FEOSSDKManager::LogFriendsInfo(const EOS_HPlatform Platform, const EOS_Epic
 		UE_LOG_EOSSDK_INFO("Friend=%d", FriendIndex);
 		Indent++;
 
-		static_assert(EOS_FRIENDS_GETSTATUS_API_LATEST == 1, "EOS_Friends_GetStatusOptions updated");
 		EOS_Friends_GetStatusOptions GetStatusOptions;
-		GetStatusOptions.ApiVersion = EOS_FRIENDS_GETSTATUS_API_LATEST;
+		GetStatusOptions.ApiVersion = 1;
+		UE_EOS_CHECK_API_MISMATCH(EOS_FRIENDS_GETSTATUS_API_LATEST, 1);
 		GetStatusOptions.LocalUserId = LoggedInAccount;
 		GetStatusOptions.TargetUserId = FriendId;
 
@@ -989,9 +989,9 @@ void FEOSSDKManager::LogConnectInfo(const EOS_HPlatform Platform, const EOS_Prod
 	EOS_EResult Result;
 
 #if !UE_BUILD_SHIPPING
-	static_assert(EOS_CONNECT_COPYIDTOKEN_API_LATEST == 1, "EOS_Connect_CopyIdTokenOptions updated");
 	EOS_Connect_CopyIdTokenOptions CopyIdTokenOptions;
-	CopyIdTokenOptions.ApiVersion = EOS_CONNECT_COPYIDTOKEN_API_LATEST;
+	CopyIdTokenOptions.ApiVersion = 1;
+	UE_EOS_CHECK_API_MISMATCH(EOS_CONNECT_COPYIDTOKEN_API_LATEST, 1);
 	CopyIdTokenOptions.LocalUserId = LoggedInAccount;
 
 	EOS_Connect_IdToken* IdToken;
@@ -1007,17 +1007,17 @@ void FEOSSDKManager::LogConnectInfo(const EOS_HPlatform Platform, const EOS_Prod
 	}
 #endif // !UE_BUILD_SHIPPING
 
-	static_assert(EOS_CONNECT_GETPRODUCTUSEREXTERNALACCOUNTCOUNT_API_LATEST == 1, "EOS_Connect_GetProductUserExternalAccountCountOptions updated");
 	EOS_Connect_GetProductUserExternalAccountCountOptions ExternalAccountCountOptions;
-	ExternalAccountCountOptions.ApiVersion = EOS_CONNECT_GETPRODUCTUSEREXTERNALACCOUNTCOUNT_API_LATEST;
+	ExternalAccountCountOptions.ApiVersion = 1;
+	UE_EOS_CHECK_API_MISMATCH(EOS_CONNECT_GETPRODUCTUSEREXTERNALACCOUNTCOUNT_API_LATEST, 1);
 	ExternalAccountCountOptions.TargetUserId = LoggedInAccount;
 
 	const uint32_t ExternalAccountCount = EOS_Connect_GetProductUserExternalAccountCount(ConnectHandle, &ExternalAccountCountOptions);
 	UE_LOG_EOSSDK_INFO("ExternalAccounts=%d", ExternalAccountCount);
 
-	static_assert(EOS_CONNECT_COPYPRODUCTUSEREXTERNALACCOUNTBYINDEX_API_LATEST == 1, "EOS_Connect_CopyProductUserExternalAccountByIndexOptions updated");
 	EOS_Connect_CopyProductUserExternalAccountByIndexOptions ExternalAccountByIndexOptions;
-	ExternalAccountByIndexOptions.ApiVersion = EOS_CONNECT_COPYPRODUCTUSEREXTERNALACCOUNTBYINDEX_API_LATEST;
+	ExternalAccountByIndexOptions.ApiVersion = 1;
+	UE_EOS_CHECK_API_MISMATCH(EOS_CONNECT_COPYPRODUCTUSEREXTERNALACCOUNTBYINDEX_API_LATEST, 1);
 	ExternalAccountByIndexOptions.TargetUserId = LoggedInAccount;
 
 	for (uint32_t ExternalAccountIndex = 0; ExternalAccountIndex < ExternalAccountCount; ExternalAccountIndex++)
diff --git a/Engine/Plugins/Online/EOSShared/Source/EOSShared/Private/Windows/WindowsEOSSDKManager.cpp b/Engine/Plugins/Online/EOSShared/Source/EOSShared/Private/Windows/WindowsEOSSDKManager.cpp
--- a/Engine/Plugins/Online/EOSShared/Source/EOSShared/Private/Windows/WindowsEOSSDKManager.cpp
+++ b/Engine/Plugins/Online/EOSShared/Source/EOSShared/Private/Windows/WindowsEOSSDKManager.cpp
@@ -6,6 +6,7 @@
 
 #include "HAL/FileManager.h"
 #include "Misc/Paths.h"
+#include "EOSShared.h"
 
 #include "Windows/eos_Windows.h"
 
@@ -20,9 +21,9 @@ IEOSPlatformHandlePtr FWindowsEOSSDKManager::CreatePlatform(const FEOSSDKPlatfor
 		static const FString XAudioPath = IFileManager::Get().ConvertToAbsolutePathForExternalAppForRead(*FPaths::Combine(FPaths::EngineDir(), TEXT("Binaries/ThirdParty/Windows/XAudio2_9"), PLATFORM_64BITS ? TEXT("x64") : TEXT("x86"), TEXT("xaudio2_9redist.dll")));
 		static const FTCHARToUTF8 Utf8XAudioPath(*XAudioPath);
 
-		static_assert(EOS_WINDOWS_RTCOPTIONS_API_LATEST == 1, "EOS_Windows_RTCOptions updated");
 		static EOS_Windows_RTCOptions WindowsRTCOptions = {};
-		WindowsRTCOptions.ApiVersion = EOS_WINDOWS_RTCOPTIONS_API_LATEST;
+		WindowsRTCOptions.ApiVersion = 1;
+		UE_EOS_CHECK_API_MISMATCH(EOS_WINDOWS_RTCOPTIONS_API_LATEST, 1);
 		WindowsRTCOptions.XAudio29DllPath = Utf8XAudioPath.Get();
 
 		const_cast<EOS_Platform_RTCOptions*>(PlatformOptions.RTCOptions)->PlatformSpecificOptions = &WindowsRTCOptions;
diff --git a/Engine/Plugins/Online/EOSShared/Source/EOSShared/Public/EOSShared.h b/Engine/Plugins/Online/EOSShared/Source/EOSShared/Public/EOSShared.h
--- a/Engine/Plugins/Online/EOSShared/Source/EOSShared/Public/EOSShared.h
+++ b/Engine/Plugins/Online/EOSShared/Source/EOSShared/Public/EOSShared.h
@@ -4,6 +4,31 @@
 
 #if WITH_EOS_SDK
 
+#include "HAL/PreprocessorHelpers.h"
+
+#ifndef UE_WITH_EOS_SDK_APIVERSION_WARNINGS
+#define UE_WITH_EOS_SDK_APIVERSION_WARNINGS 1
+#endif
+
+/**
+ * Emits warnings when an SDK upgrade bumps the ApiVersion, to let you know something changed and you should at least check the changes.
+ * To silence warnings, either bump the Actual value to match the new Expected once you have checked the API changes, or in the Engine ini hierarchy, set
+ * 
+ * [EOSShared]
+ * bEnableApiVersionWarnings=false
+ */
+#if UE_WITH_EOS_SDK_APIVERSION_WARNINGS
+#define UE_EOS_CHECK_API_MISMATCH(Actual, Expected) \
+	struct PREPROCESSOR_JOIN(FEosApiMismatchMsg_, __LINE__) { \
+		[[deprecated(#Actual " updated from " #Expected " to " PREPROCESSOR_TO_STRING(Actual) ", behaviour/params may have changed")]] \
+		static constexpr int condition(TStaticDeprecateExpression<true>) { return 1; } \
+		static constexpr int condition(TStaticDeprecateExpression<false>) { return 1; } \
+	}; \
+	enum class PREPROCESSOR_JOIN(EEosApiMismatchMsg_, __LINE__) { Value = PREPROCESSOR_JOIN(FEosApiMismatchMsg_, __LINE__)::condition(TStaticDeprecateExpression<Actual != Expected>()) }
+#else
+#define UE_EOS_CHECK_API_MISMATCH(Actual, Expected)
+#endif
+
 #include "Containers/UnrealString.h"
 #include "Logging/LogMacros.h"
 
diff --git a/Engine/Plugins/Online/OnlineServicesEOS/Source/Private/Online/AuthEOS.cpp b/Engine/Plugins/Online/OnlineServicesEOS/Source/Private/Online/AuthEOS.cpp
--- a/Engine/Plugins/Online/OnlineServicesEOS/Source/Private/Online/AuthEOS.cpp
+++ b/Engine/Plugins/Online/OnlineServicesEOS/Source/Private/Online/AuthEOS.cpp
@@ -202,10 +202,10 @@ TOnlineAsyncOpHandle<FAuthLogin> FAuthEOS::Login(FAuthLogin::Params&& Params)
 
 		// Get display name
 		EOS_UserInfo_CopyUserInfoOptions Options = { };
-		Options.ApiVersion = EOS_USERINFO_COPYUSERINFO_API_LATEST;
+		Options.ApiVersion = 3;
+		UE_EOS_CHECK_API_MISMATCH(EOS_USERINFO_COPYUSERINFO_API_LATEST, 3);
 		Options.LocalUserId = AccountInfoEOS->EpicAccountId;
 		Options.TargetUserId = AccountInfoEOS->EpicAccountId;
-		static_assert(EOS_USERINFO_COPYUSERINFO_API_LATEST == 3, "EOS_UserInfo_CopyUserInfoOptions updated, check new fields");
 
 		EOS_UserInfo* UserInfo = nullptr;
 
@@ -291,8 +291,8 @@ TOnlineAsyncOpHandle<FAuthQueryExternalServerAuthTicket> FAuthEOS::QueryExternal
 			}
 
 			EOS_Auth_CopyUserAuthTokenOptions CopyUserAuthTokenOptions = {};
-			CopyUserAuthTokenOptions.ApiVersion = EOS_AUTH_COPYUSERAUTHTOKEN_API_LATEST;
-			static_assert(EOS_AUTH_COPYUSERAUTHTOKEN_API_LATEST == 1, "EOS_Auth_CopyUserAuthTokenOptions updated, check new fields");
+			CopyUserAuthTokenOptions.ApiVersion = 1;
+			UE_EOS_CHECK_API_MISMATCH(EOS_AUTH_COPYUSERAUTHTOKEN_API_LATEST, 1);
 
 			EOS_Auth_Token* AuthToken = nullptr;
 
@@ -417,12 +417,12 @@ TFuture<TArray<FAccountId>> FAuthEOS::ResolveAccountIds(const FAccountId& LocalA
 	Algo::Transform(EpicAccountIdStrsToQuery, EpicAccountIdStrPtrs, [](const FEpicAccountIdStrBuffer& Str) { return &Str[0]; });
 
 	EOS_Connect_QueryExternalAccountMappingsOptions Options = {};
-	Options.ApiVersion = EOS_CONNECT_QUERYEXTERNALACCOUNTMAPPINGS_API_LATEST;
+	Options.ApiVersion = 1;
+	UE_EOS_CHECK_API_MISMATCH(EOS_CONNECT_QUERYEXTERNALACCOUNTMAPPINGS_API_LATEST, 1);
 	Options.LocalUserId = GetProductUserIdChecked(LocalAccountId);
 	Options.AccountIdType = EOS_EExternalAccountType::EOS_EAT_EPIC;
 	Options.ExternalAccountIds = (const char**)EpicAccountIdStrPtrs.GetData();
 	Options.ExternalAccountIdCount = 1;
-	static_assert(EOS_CONNECT_QUERYEXTERNALACCOUNTMAPPINGS_API_LATEST == 1, "EOS_Connect_QueryExternalAccountMappingsOptions updated, check new fields");
 
 	EOS_Async(EOS_Connect_QueryExternalAccountMappings, ConnectHandle, Options,
 	[this, WeakThis = AsWeak(), InEpicAccountIds, Promise = MoveTemp(Promise)](const EOS_Connect_QueryExternalAccountMappingsCallbackInfo* Data) mutable -> void
@@ -434,10 +434,10 @@ TFuture<TArray<FAccountId>> FAuthEOS::ResolveAccountIds(const FAccountId& LocalA
 			if (Data->ResultCode == EOS_EResult::EOS_Success)
 			{
 				EOS_Connect_GetExternalAccountMappingsOptions Options = {};
-				Options.ApiVersion = EOS_CONNECT_GETEXTERNALACCOUNTMAPPING_API_LATEST;
+				Options.ApiVersion = 1;
+				UE_EOS_CHECK_API_MISMATCH(EOS_CONNECT_GETEXTERNALACCOUNTMAPPING_API_LATEST, 1);
 				Options.LocalUserId = Data->LocalUserId;
 				Options.AccountIdType = EOS_EExternalAccountType::EOS_EAT_EPIC;
-				static_assert(EOS_CONNECT_GETEXTERNALACCOUNTMAPPING_API_LATEST == 1, "EOS_Connect_GetExternalAccountMappingsOptions updated, check new fields");
 
 				for (const EOS_EpicAccountId EpicAccountId : InEpicAccountIds)
 				{
@@ -503,11 +503,11 @@ TFuture<TArray<FAccountId>> FAuthEOS::ResolveAccountIds(const FAccountId& LocalA
 	TFuture<TArray<FAccountId>> Future = Promise.GetFuture();
 
 	EOS_Connect_QueryProductUserIdMappingsOptions Options = {};
-	Options.ApiVersion = EOS_CONNECT_QUERYPRODUCTUSERIDMAPPINGS_API_LATEST;
+	Options.ApiVersion = 2;
+	UE_EOS_CHECK_API_MISMATCH(EOS_CONNECT_QUERYPRODUCTUSERIDMAPPINGS_API_LATEST, 2);
 	Options.LocalUserId = GetProductUserIdChecked(LocalAccountId);
 	Options.ProductUserIds = MissingProductUserIds.GetData();
 	Options.ProductUserIdCount = MissingProductUserIds.Num();
-	static_assert(EOS_CONNECT_QUERYPRODUCTUSERIDMAPPINGS_API_LATEST == 2, "EOS_Connect_QueryProductUserIdMappingsOptions updated, check new fields");
 	
 	EOS_Async(EOS_Connect_QueryProductUserIdMappings, ConnectHandle, Options,
 	[this, WeakThis = AsWeak(), InProductUserIds, Promise = MoveTemp(Promise)](const EOS_Connect_QueryProductUserIdMappingsCallbackInfo* Data) mutable -> void
@@ -518,10 +518,10 @@ TFuture<TArray<FAccountId>> FAuthEOS::ResolveAccountIds(const FAccountId& LocalA
 			if (Data->ResultCode == EOS_EResult::EOS_Success)
 			{
 				EOS_Connect_GetProductUserIdMappingOptions Options = {};
-				Options.ApiVersion = EOS_CONNECT_GETPRODUCTUSERIDMAPPING_API_LATEST;
+				Options.ApiVersion = 1;
+				UE_EOS_CHECK_API_MISMATCH(EOS_CONNECT_GETPRODUCTUSERIDMAPPING_API_LATEST, 1);
 				Options.LocalUserId = Data->LocalUserId;
 				Options.AccountIdType = EOS_EExternalAccountType::EOS_EAT_EPIC;
-				static_assert(EOS_CONNECT_GETPRODUCTUSERIDMAPPING_API_LATEST == 1, "EOS_Connect_GetProductUserIdMappingOptions updated, check new fields");
 
 				for (const EOS_ProductUserId ProductUserId : InProductUserIds)
 				{
diff --git a/Engine/Plugins/Online/OnlineServicesEOS/Source/Private/Online/CommerceEOS.cpp b/Engine/Plugins/Online/OnlineServicesEOS/Source/Private/Online/CommerceEOS.cpp
--- a/Engine/Plugins/Online/OnlineServicesEOS/Source/Private/Online/CommerceEOS.cpp
+++ b/Engine/Plugins/Online/OnlineServicesEOS/Source/Private/Online/CommerceEOS.cpp
@@ -77,7 +77,8 @@ TOnlineAsyncOpHandle<FCommerceQueryOffers> FCommerceEOS::QueryOffers(FCommerceQu
 		{
 			const FCommerceQueryOffers::Params& Params = Op.GetParams();
 			EOS_Ecom_QueryOffersOptions Options = { };
-			Options.ApiVersion = EOS_ECOM_QUERYOFFERS_API_LATEST;
+			Options.ApiVersion = 1;
+			UE_EOS_CHECK_API_MISMATCH(EOS_ECOM_QUERYOFFERS_API_LATEST, 1);
 			Options.LocalUserId = GetEpicAccountIdChecked(Params.LocalAccountId);
 			EOS_Async(EOS_Ecom_QueryOffers, EcomHandle, Options, MoveTemp(Promise));
 		})
@@ -91,7 +92,8 @@ TOnlineAsyncOpHandle<FCommerceQueryOffers> FCommerceEOS::QueryOffers(FCommerceQu
 			}
 
 			EOS_Ecom_GetOfferCountOptions CountOptions = { };
-			CountOptions.ApiVersion = EOS_ECOM_GETOFFERCOUNT_API_LATEST;
+			CountOptions.ApiVersion = 1;
+			UE_EOS_CHECK_API_MISMATCH(EOS_ECOM_GETOFFERCOUNT_API_LATEST, 1);
 			CountOptions.LocalUserId = Data->LocalUserId;
 			uint32 OfferCount = EOS_Ecom_GetOfferCount(EcomHandle, &CountOptions);
 
@@ -99,7 +101,8 @@ TOnlineAsyncOpHandle<FCommerceQueryOffers> FCommerceEOS::QueryOffers(FCommerceQu
 			Offers.Empty(OfferCount);
 
 			EOS_Ecom_CopyOfferByIndexOptions OfferOptions = { };
-			OfferOptions.ApiVersion = EOS_ECOM_COPYOFFERBYINDEX_API_LATEST;
+			OfferOptions.ApiVersion = 3;
+			UE_EOS_CHECK_API_MISMATCH(EOS_ECOM_COPYOFFERBYINDEX_API_LATEST, 3);
 			OfferOptions.LocalUserId = Data->LocalUserId;
 			// Iterate and parse the offer list
 			for (uint32 OfferIndex = 0; OfferIndex < OfferCount; OfferIndex++)
@@ -207,12 +210,14 @@ TOnlineAsyncOpHandle<FCommerceCheckout> FCommerceEOS::Checkout(FCommerceCheckout
 			Utf8CheckoutIds.Emplace(*Offer.OfferId);
 
 			EosCheckoutEntries.AddDefaulted();
-			EosCheckoutEntries.Last().ApiVersion = EOS_ECOM_CHECKOUTENTRY_API_LATEST;
+			EosCheckoutEntries.Last().ApiVersion = 1;
+			UE_EOS_CHECK_API_MISMATCH(EOS_ECOM_CHECKOUTENTRY_API_LATEST, 1);
 			EosCheckoutEntries.Last().OfferId = Utf8CheckoutIds.Last().Get();
 		}
 
 		EOS_Ecom_CheckoutOptions Options = { };
-		Options.ApiVersion = EOS_ECOM_CHECKOUT_API_LATEST;
+		Options.ApiVersion = 1;
+		UE_EOS_CHECK_API_MISMATCH(EOS_ECOM_CHECKOUT_API_LATEST, 1);
 		Options.LocalUserId = LocalUserEasId;
 		Options.EntryCount = Params.Offers.Num();
 		Options.Entries = EosCheckoutEntries.GetData();
@@ -261,7 +266,8 @@ TOnlineAsyncOpHandle<FCommerceQueryTransactionEntitlements> FCommerceEOS::QueryT
 			EOS_Ecom_HTransaction OutTransaction;
 			FTCHARToUTF8 Utf8TransactionId(*Params.TransactionId);
 			EOS_Ecom_CopyTransactionByIdOptions CopyTransactionOptions = { };
-			CopyTransactionOptions.ApiVersion = EOS_ECOM_COPYTRANSACTIONBYID_API_LATEST;
+			CopyTransactionOptions.ApiVersion = 1;
+			UE_EOS_CHECK_API_MISMATCH(EOS_ECOM_COPYTRANSACTIONBYID_API_LATEST, 1);
 			CopyTransactionOptions.LocalUserId = LocalUserEasId;
 			CopyTransactionOptions.TransactionId = Utf8TransactionId.Get();
 
@@ -274,7 +280,8 @@ TOnlineAsyncOpHandle<FCommerceQueryTransactionEntitlements> FCommerceEOS::QueryT
 			}
 
 			EOS_Ecom_Transaction_GetEntitlementsCountOptions EntitlementCountOptions = { };
-			EntitlementCountOptions.ApiVersion = EOS_ECOM_TRANSACTION_GETENTITLEMENTSCOUNT_API_LATEST;
+			EntitlementCountOptions.ApiVersion = 1;
+			UE_EOS_CHECK_API_MISMATCH(EOS_ECOM_TRANSACTION_GETENTITLEMENTSCOUNT_API_LATEST, 1);
 			int32 NumTransactionsFound = EOS_Ecom_Transaction_GetEntitlementsCount(OutTransaction, &EntitlementCountOptions);
 
 			if (NumTransactionsFound > 0)
@@ -285,7 +292,8 @@ TOnlineAsyncOpHandle<FCommerceQueryTransactionEntitlements> FCommerceEOS::QueryT
 				{
 					EOS_Ecom_Entitlement* EosEntitlement;
 					EOS_Ecom_Transaction_CopyEntitlementByIndexOptions CopyTransactionEntitlementOptions = {};
-					CopyTransactionEntitlementOptions.ApiVersion = EOS_ECOM_TRANSACTION_COPYENTITLEMENTBYINDEX_API_LATEST;
+					CopyTransactionEntitlementOptions.ApiVersion = 1;
+					UE_EOS_CHECK_API_MISMATCH(EOS_ECOM_TRANSACTION_COPYENTITLEMENTBYINDEX_API_LATEST, 1);
 					CopyTransactionEntitlementOptions.EntitlementIndex = i;
 
 					EOS_EResult Result = EOS_Ecom_Transaction_CopyEntitlementByIndex(OutTransaction, &CopyTransactionEntitlementOptions, &EosEntitlement);
@@ -334,7 +342,8 @@ TOnlineAsyncOpHandle<FCommerceQueryEntitlements> FCommerceEOS::QueryEntitlements
 		{
 			const FCommerceQueryEntitlements::Params& Params = Op.GetParams();
 			EOS_Ecom_QueryEntitlementsOptions Options = { };
-			Options.ApiVersion = EOS_ECOM_QUERYENTITLEMENTS_API_LATEST;
+			Options.ApiVersion = 2;
+			UE_EOS_CHECK_API_MISMATCH(EOS_ECOM_QUERYENTITLEMENTS_API_LATEST, 2);
 			Options.LocalUserId = GetEpicAccountIdChecked(Params.LocalAccountId);
 			Options.bIncludeRedeemed = EOS_TRUE;
 			EOS_Async(EOS_Ecom_QueryEntitlements, EcomHandle, Options, MoveTemp(Promise));
@@ -350,14 +359,16 @@ TOnlineAsyncOpHandle<FCommerceQueryEntitlements> FCommerceEOS::QueryEntitlements
 			}
 
 			EOS_Ecom_GetEntitlementsCountOptions CountOptions = { };
-			CountOptions.ApiVersion = EOS_ECOM_GETENTITLEMENTSCOUNT_API_LATEST;
+			CountOptions.ApiVersion = 1;
+			UE_EOS_CHECK_API_MISMATCH(EOS_ECOM_GETENTITLEMENTSCOUNT_API_LATEST, 1);
 			CountOptions.LocalUserId = Data->LocalUserId;
 			uint32 Count = EOS_Ecom_GetEntitlementsCount(EcomHandle, &CountOptions);
 			TArray<FEntitlement>& Entitlements = CachedEntitlements.FindOrAdd(Op.GetParams().LocalAccountId, TArray<FEntitlement>());
 			Entitlements.Reset(Count);
 
 			EOS_Ecom_CopyEntitlementByIndexOptions CopyOptions = { };
-			CopyOptions.ApiVersion = EOS_ECOM_COPYENTITLEMENTBYINDEX_API_LATEST;
+			CopyOptions.ApiVersion = 1;
+			UE_EOS_CHECK_API_MISMATCH(EOS_ECOM_COPYENTITLEMENTBYINDEX_API_LATEST, 1);
 			CopyOptions.LocalUserId = Data->LocalUserId;
 
 			for (uint32 Index = 0; Index < Count; Index++)
@@ -418,7 +429,8 @@ TOnlineAsyncOpHandle<FCommerceRedeemEntitlement> FCommerceEOS::RedeemEntitlement
 		FTCHARToUTF8 Utf8EntitlementId(*Op.GetParams().EntitlementId);
 		const char* EntitlementIdPtr = Utf8EntitlementId.Get();
 		EOS_Ecom_RedeemEntitlementsOptions Options = { };
-		Options.ApiVersion = EOS_ECOM_REDEEMENTITLEMENTS_API_LATEST;
+		Options.ApiVersion = 2;
+		UE_EOS_CHECK_API_MISMATCH(EOS_ECOM_REDEEMENTITLEMENTS_API_LATEST, 2);
 		Options.LocalUserId = LocalUserEasId;
 		Options.EntitlementIdCount = 1;
 		Options.EntitlementIds = &EntitlementIdPtr;
@@ -465,7 +477,8 @@ TOnlineAsyncOpHandle<FCommerceRetrieveS2SToken> FCommerceEOS::RetrieveS2SToken(F
 		{
 			const FCommerceRetrieveS2SToken::Params& Params = Op.GetParams();
 			EOS_Ecom_QueryOwnershipTokenOptions Options = {};
-			Options.ApiVersion = EOS_ECOM_QUERYOWNERSHIPTOKEN_API_LATEST;
+			Options.ApiVersion = 2;
+			UE_EOS_CHECK_API_MISMATCH(EOS_ECOM_QUERYOWNERSHIPTOKEN_API_LATEST, 2);
 			Options.LocalUserId = GetEpicAccountIdChecked(Params.LocalAccountId);
 			EOS_Async(EOS_Ecom_QueryOwnershipToken, EcomHandle, Options, MoveTemp(Promise));
 		})
diff --git a/Engine/Plugins/Online/OnlineServicesEOS/Source/Private/Online/ExternalUIEOS.cpp b/Engine/Plugins/Online/OnlineServicesEOS/Source/Private/Online/ExternalUIEOS.cpp
--- a/Engine/Plugins/Online/OnlineServicesEOS/Source/Private/Online/ExternalUIEOS.cpp
+++ b/Engine/Plugins/Online/OnlineServicesEOS/Source/Private/Online/ExternalUIEOS.cpp
@@ -93,7 +93,8 @@ TOnlineAsyncOpHandle<FExternalUIShowFriendsUI> FExternalUIEOS::ShowFriendsUI(FEx
 	}
 
 	EOS_UI_ShowFriendsOptions ShowFriendsOptions = {};
-	ShowFriendsOptions.ApiVersion = EOS_UI_SHOWFRIENDS_API_LATEST;
+	ShowFriendsOptions.ApiVersion = 1;
+	UE_EOS_CHECK_API_MISMATCH(EOS_UI_SHOWFRIENDS_API_LATEST, 1);
 	ShowFriendsOptions.LocalUserId = LocalUserEasId;
 
 	Op->Then([this, ShowFriendsOptions](TOnlineAsyncOp<FExternalUIShowFriendsUI>& InAsyncOp, TPromise<const EOS_UI_ShowFriendsCallbackInfo*>&& Promise) mutable
diff --git a/Engine/Plugins/Online/OnlineServicesEOS/Source/Private/Online/PresenceEOS.cpp b/Engine/Plugins/Online/OnlineServicesEOS/Source/Private/Online/PresenceEOS.cpp
--- a/Engine/Plugins/Online/OnlineServicesEOS/Source/Private/Online/PresenceEOS.cpp
+++ b/Engine/Plugins/Online/OnlineServicesEOS/Source/Private/Online/PresenceEOS.cpp
@@ -74,7 +74,8 @@ void FPresenceEOS::Initialize()
 
 	// Register for friend updates
 	EOS_Presence_AddNotifyOnPresenceChangedOptions Options = { };
-	Options.ApiVersion = EOS_PRESENCE_ADDNOTIFYONPRESENCECHANGED_API_LATEST;
+	Options.ApiVersion = 1;
+	UE_EOS_CHECK_API_MISMATCH(EOS_PRESENCE_ADDNOTIFYONPRESENCECHANGED_API_LATEST, 1);
 	NotifyPresenceChangedNotificationId = EOS_Presence_AddNotifyOnPresenceChanged(PresenceHandle, &Options, this, [](const EOS_Presence_PresenceChangedCallbackInfo* Data)
 	{
 		FPresenceEOS* This = reinterpret_cast<FPresenceEOS*>(Data->ClientData);
@@ -164,7 +165,8 @@ TOnlineAsyncOpHandle<FQueryPresence> FPresenceEOS::QueryPresence(FQueryPresence:
 			{
 				const FQueryPresence::Params& Params = InAsyncOp.GetParams();
 				EOS_Presence_QueryPresenceOptions QueryPresenceOptions = { };
-				QueryPresenceOptions.ApiVersion = EOS_PRESENCE_QUERYPRESENCE_API_LATEST;
+				QueryPresenceOptions.ApiVersion = 1;
+				UE_EOS_CHECK_API_MISMATCH(EOS_PRESENCE_QUERYPRESENCE_API_LATEST, 1);
 				QueryPresenceOptions.LocalUserId = GetEpicAccountIdChecked(Params.LocalAccountId);
 				QueryPresenceOptions.TargetUserId = GetEpicAccountIdChecked(Params.TargetAccountId);
 
@@ -221,7 +223,8 @@ TOnlineAsyncOpHandle<FUpdatePresence> FPresenceEOS::UpdatePresence(FUpdatePresen
 		const FUpdatePresence::Params& Params = InAsyncOp.GetParams();
 		EOS_HPresenceModification ChangeHandle = nullptr;
 		EOS_Presence_CreatePresenceModificationOptions Options = { };
-		Options.ApiVersion = EOS_PRESENCE_CREATEPRESENCEMODIFICATION_API_LATEST;
+		Options.ApiVersion = 1;
+		UE_EOS_CHECK_API_MISMATCH(EOS_PRESENCE_CREATEPRESENCEMODIFICATION_API_LATEST, 1);
 		Options.LocalUserId = GetEpicAccountIdChecked(Params.LocalAccountId);
 		EOS_EResult CreatePresenceModificationResult = EOS_Presence_CreatePresenceModification(PresenceHandle, &Options, &ChangeHandle);
 		if (CreatePresenceModificationResult == EOS_EResult::EOS_Success)
@@ -230,7 +233,8 @@ TOnlineAsyncOpHandle<FUpdatePresence> FPresenceEOS::UpdatePresence(FUpdatePresen
 
 			// State
 			EOS_PresenceModification_SetStatusOptions StatusOptions = { };
-			StatusOptions.ApiVersion = EOS_PRESENCEMODIFICATION_SETSTATUS_API_LATEST;
+			StatusOptions.ApiVersion = 1;
+			UE_EOS_CHECK_API_MISMATCH(EOS_PRESENCEMODIFICATION_SETSTATUS_API_LATEST, 1);
 			StatusOptions.Status = ToEOS_Presence_EStatus(Params.Presence->Status);
 			EOS_EResult SetStatusResult = EOS_PresenceModification_SetStatus(ChangeHandle, &StatusOptions);
 			if (SetStatusResult != EOS_EResult::EOS_Success)
@@ -244,7 +248,8 @@ TOnlineAsyncOpHandle<FUpdatePresence> FPresenceEOS::UpdatePresence(FUpdatePresen
 			// Raw rich text
 			// Convert the status string as the rich text string
 			EOS_PresenceModification_SetRawRichTextOptions RawRichTextOptions = { };
-			RawRichTextOptions.ApiVersion = EOS_PRESENCEMODIFICATION_SETRAWRICHTEXT_API_LATEST;
+			RawRichTextOptions.ApiVersion = 1;
+			UE_EOS_CHECK_API_MISMATCH(EOS_PRESENCEMODIFICATION_SETRAWRICHTEXT_API_LATEST, 1);
 
 			const FTCHARToUTF8 Utf8RawRichText(*Params.Presence->StatusString);
 			RawRichTextOptions.RichText = Utf8RawRichText.Get();
@@ -281,14 +286,16 @@ TOnlineAsyncOpHandle<FUpdatePresence> FPresenceEOS::UpdatePresence(FUpdatePresen
 					const FTCHARToUTF8& Utf8Key = Utf8Strings.Emplace_GetRef(*RemovedProperty);
 
 					EOS_PresenceModification_DataRecordId& RecordId = RecordIds.Emplace_GetRef();
-					RecordId.ApiVersion = EOS_PRESENCEMODIFICATION_DATARECORDID_API_LATEST;
+					RecordId.ApiVersion = 1;
+					UE_EOS_CHECK_API_MISMATCH(EOS_PRESENCEMODIFICATION_DATARECORDID_API_LATEST, 1);
 					RecordId.Key = Utf8Key.Get();
 
 					UE_LOG(LogTemp, Warning, TEXT("UpdatePresence: Removing field %s"), *RemovedProperty); // Temp logging
 				}
 
 				EOS_PresenceModification_DeleteDataOptions DataOptions = { };
-				DataOptions.ApiVersion = EOS_PRESENCEMODIFICATION_DELETEDATA_API_LATEST;
+				DataOptions.ApiVersion = 1;
+				UE_EOS_CHECK_API_MISMATCH(EOS_PRESENCEMODIFICATION_DELETEDATA_API_LATEST, 1);
 				DataOptions.RecordsCount = RecordIds.Num();
 				DataOptions.Records = RecordIds.GetData();
 				EOS_EResult DeleteDataResult = EOS_PresenceModification_DeleteData(ChangeHandle, &DataOptions);
@@ -321,14 +328,16 @@ TOnlineAsyncOpHandle<FUpdatePresence> FPresenceEOS::UpdatePresence(FUpdatePresen
 					const FTCHARToUTF8& Utf8Value = Utf8Strings.Emplace_GetRef(*UpdatedProperty.Value); // TODO: Better serialization
 
 					EOS_Presence_DataRecord& Record = Records.Emplace_GetRef();
-					Record.ApiVersion = EOS_PRESENCE_DATARECORD_API_LATEST;
+					Record.ApiVersion = 1;
+					UE_EOS_CHECK_API_MISMATCH(EOS_PRESENCE_DATARECORD_API_LATEST, 1);
 					Record.Key = Utf8Key.Get();
 					Record.Value = Utf8Value.Get();
 					UE_LOG(LogTemp, Warning, TEXT("UpdatePresence: Set field [%s] to [%s]"), *UpdatedProperty.Key, *UpdatedProperty.Value); // Temp logging
 				}
 
 				EOS_PresenceModification_SetDataOptions DataOptions = { };
-				DataOptions.ApiVersion = EOS_PRESENCEMODIFICATION_SETDATA_API_LATEST;
+				DataOptions.ApiVersion = 1;
+				UE_EOS_CHECK_API_MISMATCH(EOS_PRESENCEMODIFICATION_SETDATA_API_LATEST, 1);
 				DataOptions.RecordsCount = Records.Num();
 				DataOptions.Records = Records.GetData();
 				EOS_EResult SetDataResult = EOS_PresenceModification_SetData(ChangeHandle, &DataOptions);
@@ -342,7 +351,8 @@ TOnlineAsyncOpHandle<FUpdatePresence> FPresenceEOS::UpdatePresence(FUpdatePresen
 			}
 
 			EOS_Presence_SetPresenceOptions SetPresenceOptions = { };
-			SetPresenceOptions.ApiVersion = EOS_PRESENCE_SETPRESENCE_API_LATEST;
+			SetPresenceOptions.ApiVersion = 1;
+			UE_EOS_CHECK_API_MISMATCH(EOS_PRESENCE_SETPRESENCE_API_LATEST, 1);
 			SetPresenceOptions.LocalUserId = GetEpicAccountIdChecked(Params.LocalAccountId);
 			SetPresenceOptions.PresenceModificationHandle = ChangeHandle;
 
@@ -404,7 +414,8 @@ TOnlineAsyncOpHandle<FPartialUpdatePresence> FPresenceEOS::PartialUpdatePresence
 				const FPartialUpdatePresence::Params& Params = InAsyncOp.GetParams();
 				EOS_HPresenceModification ChangeHandle = nullptr;
 				EOS_Presence_CreatePresenceModificationOptions Options = { };
-				Options.ApiVersion = EOS_PRESENCE_CREATEPRESENCEMODIFICATION_API_LATEST;
+				Options.ApiVersion = 1;
+				UE_EOS_CHECK_API_MISMATCH(EOS_PRESENCE_CREATEPRESENCEMODIFICATION_API_LATEST, 1);
 				Options.LocalUserId = GetEpicAccountIdChecked(Params.LocalAccountId);
 				EOS_EResult CreatePresenceModificationResult = EOS_Presence_CreatePresenceModification(PresenceHandle, &Options, &ChangeHandle);
 				if (CreatePresenceModificationResult == EOS_EResult::EOS_Success)
@@ -415,7 +426,8 @@ TOnlineAsyncOpHandle<FPartialUpdatePresence> FPresenceEOS::PartialUpdatePresence
 					if (Params.Mutations.Status.IsSet())
 					{
 						EOS_PresenceModification_SetStatusOptions StatusOptions = { };
-						StatusOptions.ApiVersion = EOS_PRESENCEMODIFICATION_SETSTATUS_API_LATEST;
+						StatusOptions.ApiVersion = 1;
+						UE_EOS_CHECK_API_MISMATCH(EOS_PRESENCEMODIFICATION_SETSTATUS_API_LATEST, 1);
 						StatusOptions.Status = ToEOS_Presence_EStatus(Params.Mutations.Status.GetValue());
 						EOS_EResult SetStatusResult = EOS_PresenceModification_SetStatus(ChangeHandle, &StatusOptions);
 						if (SetStatusResult != EOS_EResult::EOS_Success)
@@ -436,7 +448,8 @@ TOnlineAsyncOpHandle<FPartialUpdatePresence> FPresenceEOS::PartialUpdatePresence
 					{
 						// Convert the status string as the rich text string
 						EOS_PresenceModification_SetRawRichTextOptions RawRichTextOptions = { };
-						RawRichTextOptions.ApiVersion = EOS_PRESENCEMODIFICATION_SETRAWRICHTEXT_API_LATEST;
+						RawRichTextOptions.ApiVersion = 1;
+						UE_EOS_CHECK_API_MISMATCH(EOS_PRESENCEMODIFICATION_SETRAWRICHTEXT_API_LATEST, 1);
 
 						const FTCHARToUTF8 Utf8RawRichText(*Params.Mutations.StatusString.GetValue());
 						RawRichTextOptions.RichText = Utf8RawRichText.Get();
@@ -466,14 +479,16 @@ TOnlineAsyncOpHandle<FPartialUpdatePresence> FPresenceEOS::PartialUpdatePresence
 							const FTCHARToUTF8& Utf8Key = Utf8Strings.Emplace_GetRef(*RemovedProperty);
 
 							EOS_PresenceModification_DataRecordId& RecordId = RecordIds.Emplace_GetRef();
-							RecordId.ApiVersion = EOS_PRESENCEMODIFICATION_DATARECORDID_API_LATEST;
+							RecordId.ApiVersion = 1;
+							UE_EOS_CHECK_API_MISMATCH(EOS_PRESENCEMODIFICATION_DATARECORDID_API_LATEST, 1);
 							RecordId.Key = Utf8Key.Get();
 
 							UE_LOG(LogTemp, Warning, TEXT("UpdatePresence: Removing field %s"), *RemovedProperty); // Temp logging
 						}
 
 						EOS_PresenceModification_DeleteDataOptions DataOptions = { };
-						DataOptions.ApiVersion = EOS_PRESENCEMODIFICATION_DELETEDATA_API_LATEST;
+						DataOptions.ApiVersion = 1;
+						UE_EOS_CHECK_API_MISMATCH(EOS_PRESENCEMODIFICATION_DELETEDATA_API_LATEST, 1);
 						DataOptions.RecordsCount = RecordIds.Num();
 						DataOptions.Records = RecordIds.GetData();
 						EOS_EResult DeleteDataResult = EOS_PresenceModification_DeleteData(ChangeHandle, &DataOptions);
@@ -506,14 +521,16 @@ TOnlineAsyncOpHandle<FPartialUpdatePresence> FPresenceEOS::PartialUpdatePresence
 							const FTCHARToUTF8& Utf8Value = Utf8Strings.Emplace_GetRef(*UpdatedProperty.Value); // TODO: Better serialization
 
 							EOS_Presence_DataRecord& Record = Records.Emplace_GetRef();
-							Record.ApiVersion = EOS_PRESENCE_DATARECORD_API_LATEST;
+							Record.ApiVersion = 1;
+							UE_EOS_CHECK_API_MISMATCH(EOS_PRESENCE_DATARECORD_API_LATEST, 1);
 							Record.Key = Utf8Key.Get();
 							Record.Value = Utf8Value.Get();
 							UE_LOG(LogTemp, Warning, TEXT("UpdatePresence: Set field [%s] to [%s]"), *UpdatedProperty.Key, *UpdatedProperty.Value); // Temp logging
 						}
 
 						EOS_PresenceModification_SetDataOptions DataOptions = { };
-						DataOptions.ApiVersion = EOS_PRESENCEMODIFICATION_SETDATA_API_LATEST;
+						DataOptions.ApiVersion = 1;
+						UE_EOS_CHECK_API_MISMATCH(EOS_PRESENCEMODIFICATION_SETDATA_API_LATEST, 1);
 						DataOptions.RecordsCount = Records.Num();
 						DataOptions.Records = Records.GetData();
 						EOS_EResult SetDataResult = EOS_PresenceModification_SetData(ChangeHandle, &DataOptions);
@@ -527,7 +544,8 @@ TOnlineAsyncOpHandle<FPartialUpdatePresence> FPresenceEOS::PartialUpdatePresence
 					}
 
 					EOS_Presence_SetPresenceOptions SetPresenceOptions = { };
-					SetPresenceOptions.ApiVersion = EOS_PRESENCE_SETPRESENCE_API_LATEST;
+					SetPresenceOptions.ApiVersion = 1;
+					UE_EOS_CHECK_API_MISMATCH(EOS_PRESENCE_SETPRESENCE_API_LATEST, 1);
 					SetPresenceOptions.LocalUserId = GetEpicAccountIdChecked(Params.LocalAccountId);
 					SetPresenceOptions.PresenceModificationHandle = ChangeHandle;
 					EOS_Async(EOS_Presence_SetPresence, PresenceHandle, SetPresenceOptions, MoveTemp(Promise));
@@ -610,7 +628,8 @@ void FPresenceEOS::UpdateUserPresence(FAccountId LocalAccountId, FAccountId Pres
 	// Get presence from EOS
 	EOS_Presence_Info* PresenceInfo = nullptr;
 	EOS_Presence_CopyPresenceOptions Options = { };
-	Options.ApiVersion = EOS_PRESENCE_COPYPRESENCE_API_LATEST;
+	Options.ApiVersion = 3;
+	UE_EOS_CHECK_API_MISMATCH(EOS_PRESENCE_COPYPRESENCE_API_LATEST, 3);
 	Options.LocalUserId = GetEpicAccountIdChecked(LocalAccountId);
 	Options.TargetUserId = GetEpicAccountIdChecked(PresenceAccountId);
 	EOS_EResult CopyPresenceResult = EOS_Presence_CopyPresence(PresenceHandle, &Options, &PresenceInfo);
diff --git a/Engine/Plugins/Online/OnlineServicesEOS/Source/Private/Online/SocialEOS.cpp b/Engine/Plugins/Online/OnlineServicesEOS/Source/Private/Online/SocialEOS.cpp
--- a/Engine/Plugins/Online/OnlineServicesEOS/Source/Private/Online/SocialEOS.cpp
+++ b/Engine/Plugins/Online/OnlineServicesEOS/Source/Private/Online/SocialEOS.cpp
@@ -41,7 +41,8 @@ void FSocialEOS::Initialize()
 
 	// Register for friend updates
 	EOS_Friends_AddNotifyFriendsUpdateOptions Options = { };
-	Options.ApiVersion = EOS_FRIENDS_ADDNOTIFYFRIENDSUPDATE_API_LATEST;
+	Options.ApiVersion = 1;
+	UE_EOS_CHECK_API_MISMATCH(EOS_FRIENDS_ADDNOTIFYFRIENDSUPDATE_API_LATEST, 1);
 	NotifyFriendsUpdateNotificationId = EOS_Friends_AddNotifyFriendsUpdate(FriendsHandle, &Options, this, [](const EOS_Friends_OnFriendsUpdateInfo* Data)
 	{
 		FSocialEOS* This = reinterpret_cast<FSocialEOS*>(Data->ClientData);
@@ -87,7 +88,8 @@ TOnlineAsyncOpHandle<FQueryFriends> FSocialEOS::QueryFriends(FQueryFriends::Para
 	.Then([this](TOnlineAsyncOp<FQueryFriends>& Op, TPromise<const EOS_Friends_QueryFriendsCallbackInfo*>&& Promise)
 	{
 		EOS_Friends_QueryFriendsOptions QueryFriendsOptions = { };
-		QueryFriendsOptions.ApiVersion = EOS_FRIENDS_QUERYFRIENDS_API_LATEST;
+		QueryFriendsOptions.ApiVersion = 1;
+		UE_EOS_CHECK_API_MISMATCH(EOS_FRIENDS_QUERYFRIENDS_API_LATEST, 1);
 		QueryFriendsOptions.LocalUserId = GetEpicAccountIdChecked(Op.GetParams().LocalAccountId);
 
 		EOS_Async(EOS_Friends_QueryFriends, FriendsHandle, QueryFriendsOptions, MoveTemp(Promise));
@@ -212,7 +214,8 @@ TOnlineAsyncOpHandle<FSendFriendInvite> FSocialEOS::SendFriendInvite(FSendFriend
 		const FSendFriendInvite::Params& Params = Op.GetParams();
 
 		EOS_Friends_SendInviteOptions SendInviteOptions = {};
-		SendInviteOptions.ApiVersion = EOS_FRIENDS_SENDINVITE_API_LATEST;
+		SendInviteOptions.ApiVersion = 1;
+		UE_EOS_CHECK_API_MISMATCH(EOS_FRIENDS_SENDINVITE_API_LATEST, 1);
 		SendInviteOptions.LocalUserId = GetEpicAccountIdChecked(Params.LocalAccountId);
 		SendInviteOptions.TargetUserId = GetEpicAccountIdChecked(Params.TargetAccountId);
 
@@ -263,7 +266,8 @@ TOnlineAsyncOpHandle<FAcceptFriendInvite> FSocialEOS::AcceptFriendInvite(FAccept
 		const FAcceptFriendInvite::Params& Params = Op.GetParams();
 
 		EOS_Friends_AcceptInviteOptions AcceptInviteOptions = {};
-		AcceptInviteOptions.ApiVersion = EOS_FRIENDS_ACCEPTINVITE_API_LATEST;
+		AcceptInviteOptions.ApiVersion = 1;
+		UE_EOS_CHECK_API_MISMATCH(EOS_FRIENDS_ACCEPTINVITE_API_LATEST, 1);
 		AcceptInviteOptions.LocalUserId = GetEpicAccountIdChecked(Params.LocalAccountId);
 		AcceptInviteOptions.TargetUserId = GetEpicAccountIdChecked(Params.TargetAccountId);
 
@@ -314,7 +318,8 @@ TOnlineAsyncOpHandle<FRejectFriendInvite> FSocialEOS::RejectFriendInvite(FReject
 		const FRejectFriendInvite::Params& Params = Op.GetParams();
 
 		EOS_Friends_RejectInviteOptions RejectInviteOptions = {};
-		RejectInviteOptions.ApiVersion = EOS_FRIENDS_REJECTINVITE_API_LATEST;
+		RejectInviteOptions.ApiVersion = 1;
+		UE_EOS_CHECK_API_MISMATCH(EOS_FRIENDS_REJECTINVITE_API_LATEST, 1);
 		RejectInviteOptions.LocalUserId = GetEpicAccountIdChecked(Params.LocalAccountId);
 		RejectInviteOptions.TargetUserId = GetEpicAccountIdChecked(Params.TargetAccountId);
 
@@ -371,7 +376,8 @@ TOnlineAsyncOpHandle<FBlockUser> FSocialEOS::BlockUser(FBlockUser::Params&& Para
 		const FBlockUser::Params& Params = Op.GetParams();
 
 		EOS_UI_ShowBlockPlayerOptions ShowBlockPlayerOptions = {};
-		ShowBlockPlayerOptions.ApiVersion = EOS_UI_SHOWBLOCKPLAYER_API_LATEST;
+		ShowBlockPlayerOptions.ApiVersion = 1;
+		UE_EOS_CHECK_API_MISMATCH(EOS_UI_SHOWBLOCKPLAYER_API_LATEST, 1);
 		ShowBlockPlayerOptions.LocalUserId = GetEpicAccountIdChecked(Params.LocalAccountId);
 		ShowBlockPlayerOptions.TargetUserId = GetEpicAccountIdChecked(Params.TargetAccountId);
 
diff --git a/Engine/Plugins/Online/OnlineServicesEOS/Source/Private/Online/UserInfoEOS.cpp b/Engine/Plugins/Online/OnlineServicesEOS/Source/Private/Online/UserInfoEOS.cpp
--- a/Engine/Plugins/Online/OnlineServicesEOS/Source/Private/Online/UserInfoEOS.cpp
+++ b/Engine/Plugins/Online/OnlineServicesEOS/Source/Private/Online/UserInfoEOS.cpp
@@ -60,8 +60,8 @@ TOnlineAsyncOpHandle<FQueryUserInfo> FUserInfoEOS::QueryUserInfo(FQueryUserInfo:
 				}
 
 				EOS_UserInfo_QueryUserInfoOptions QueryUserInfoOptions = {};
-				QueryUserInfoOptions.ApiVersion = EOS_USERINFO_QUERYUSERINFO_API_LATEST;
-				static_assert(EOS_USERINFO_QUERYUSERINFO_API_LATEST == 1, "EOS_UserInfo_QueryUserInfoOptions updated, check new fields");
+				QueryUserInfoOptions.ApiVersion = 1;
+				UE_EOS_CHECK_API_MISMATCH(EOS_USERINFO_QUERYUSERINFO_API_LATEST, 1);
 				QueryUserInfoOptions.LocalUserId = GetEpicAccountIdChecked(Params.LocalAccountId);
 				QueryUserInfoOptions.TargetUserId = TargetUserEasId;
 
@@ -101,8 +101,8 @@ TOnlineResult<FGetUserInfo> FUserInfoEOS::GetUserInfo(FGetUserInfo::Params&& Par
 	}
 
 	EOS_UserInfo_CopyUserInfoOptions Options;
-	Options.ApiVersion = EOS_USERINFO_COPYUSERINFO_API_LATEST;
-	static_assert(EOS_USERINFO_COPYUSERINFO_API_LATEST == 3, "EOS_UserInfo_CopyUserInfoOptions updated, check new fields");
+	Options.ApiVersion = 3;
+	UE_EOS_CHECK_API_MISMATCH(EOS_USERINFO_COPYUSERINFO_API_LATEST, 3);
 	Options.LocalUserId = GetEpicAccountIdChecked(Params.LocalAccountId);
 	Options.TargetUserId = TargetUserEasId;
 
diff --git a/Engine/Plugins/Online/OnlineServicesEOSGS/Source/Private/Online/AchievementsEOSGS.cpp b/Engine/Plugins/Online/OnlineServicesEOSGS/Source/Private/Online/AchievementsEOSGS.cpp
--- a/Engine/Plugins/Online/OnlineServicesEOSGS/Source/Private/Online/AchievementsEOSGS.cpp
+++ b/Engine/Plugins/Online/OnlineServicesEOSGS/Source/Private/Online/AchievementsEOSGS.cpp
@@ -27,8 +27,8 @@ void FAchievementsEOSGS::Initialize()
 
 	// Register for achievement unlocked events
 	EOS_Achievements_AddNotifyAchievementsUnlockedV2Options Options = { };
-	Options.ApiVersion = EOS_ACHIEVEMENTS_ADDNOTIFYACHIEVEMENTSUNLOCKEDV2_API_LATEST;
-	static_assert(EOS_ACHIEVEMENTS_ADDNOTIFYACHIEVEMENTSUNLOCKEDV2_API_LATEST == 2, "EOS_Achievements_AddNotifyAchievementsUnlockedV2Options updated, check new fields");
+	Options.ApiVersion = 2;
+	UE_EOS_CHECK_API_MISMATCH(EOS_ACHIEVEMENTS_ADDNOTIFYACHIEVEMENTSUNLOCKEDV2_API_LATEST, 2);
 	NotifyAchievementsUnlockedNotificationId = EOS_Achievements_AddNotifyAchievementsUnlockedV2(AchievementsHandle, &Options, this, [](const EOS_Achievements_OnAchievementsUnlockedCallbackV2Info* Data)
 	{
 		FAchievementsEOSGS* This = reinterpret_cast<FAchievementsEOSGS*>(Data->ClientData);
@@ -93,8 +93,8 @@ TOnlineAsyncOpHandle<FQueryAchievementDefinitions> FAchievementsEOSGS::QueryAchi
 			const FQueryAchievementDefinitions::Params& Params = InAsyncOp.GetParams();
 
 			EOS_Achievements_QueryDefinitionsOptions Options = {};
-			Options.ApiVersion = EOS_ACHIEVEMENTS_QUERYDEFINITIONS_API_LATEST;
-			static_assert(EOS_ACHIEVEMENTS_QUERYDEFINITIONS_API_LATEST == 3, "EOS_Achievements_QueryDefinitionsOptions updated, check new fields");
+			Options.ApiVersion = 3;
+			UE_EOS_CHECK_API_MISMATCH(EOS_ACHIEVEMENTS_QUERYDEFINITIONS_API_LATEST, 3);
 			Options.LocalUserId = GetProductUserIdChecked(Params.LocalAccountId);
 
 			EOS_Async(EOS_Achievements_QueryDefinitions, AchievementsHandle, Options, MoveTemp(Promise));
@@ -110,8 +110,8 @@ TOnlineAsyncOpHandle<FQueryAchievementDefinitions> FAchievementsEOSGS::QueryAchi
 			}
 
 			EOS_Achievements_GetAchievementDefinitionCountOptions GetCountOptions = {};
-			GetCountOptions.ApiVersion = EOS_ACHIEVEMENTS_GETACHIEVEMENTDEFINITIONCOUNT_API_LATEST;
-			static_assert(EOS_ACHIEVEMENTS_GETACHIEVEMENTDEFINITIONCOUNT_API_LATEST == 1, "EOS_Achievements_GetAchievementDefinitionCountOptions updated, check new fields");
+			GetCountOptions.ApiVersion = 1;
+			UE_EOS_CHECK_API_MISMATCH(EOS_ACHIEVEMENTS_GETACHIEVEMENTDEFINITIONCOUNT_API_LATEST, 1);
 
 			const uint32_t AchievementCount = EOS_Achievements_GetAchievementDefinitionCount(AchievementsHandle, &GetCountOptions);
 
@@ -119,8 +119,8 @@ TOnlineAsyncOpHandle<FQueryAchievementDefinitions> FAchievementsEOSGS::QueryAchi
 			for (uint32_t AchievementIdx = 0; AchievementIdx < AchievementCount; AchievementIdx++)
 			{
 				EOS_Achievements_CopyAchievementDefinitionV2ByIndexOptions CopyOptions = {};
-				CopyOptions.ApiVersion = EOS_ACHIEVEMENTS_COPYACHIEVEMENTDEFINITIONV2BYINDEX_API_LATEST;
-				static_assert(EOS_ACHIEVEMENTS_COPYACHIEVEMENTDEFINITIONV2BYINDEX_API_LATEST == 2, "EOS_Achievements_CopyAchievementDefinitionV2ByIndexOptions updated, check new fields");
+				CopyOptions.ApiVersion = 2;
+				UE_EOS_CHECK_API_MISMATCH(EOS_ACHIEVEMENTS_COPYACHIEVEMENTDEFINITIONV2BYINDEX_API_LATEST, 2);
 				CopyOptions.AchievementIndex = AchievementIdx;
 
 				EOS_Achievements_DefinitionV2* EOSDefinition = nullptr;
@@ -133,9 +133,10 @@ TOnlineAsyncOpHandle<FQueryAchievementDefinitions> FAchievementsEOSGS::QueryAchi
 					return;
 				}
 
-				static_assert(EOS_ACHIEVEMENTS_DEFINITIONV2_API_LATEST == 2, "EOS_Achievements_DefinitionV2 updated, check new fields");
-				const bool bDefsApiVersionOk = EOSDefinition->ApiVersion == EOS_ACHIEVEMENTS_DEFINITIONV2_API_LATEST;
-				UE_CLOG(!bDefsApiVersionOk, LogTemp, Warning, TEXT("EOS_Achievements_DefinitionV2 version mismatch Expected=%d Actual=%d"), EOS_ACHIEVEMENTS_DEFINITIONV2_API_LATEST, EOSDefinition->ApiVersion);
+				const int ExpectedDefinitionApiVersion = 2;
+				UE_EOS_CHECK_API_MISMATCH(EOS_ACHIEVEMENTS_DEFINITIONV2_API_LATEST, ExpectedDefinitionApiVersion);
+				const bool bDefsApiVersionOk = EOSDefinition->ApiVersion == ExpectedDefinitionApiVersion;
+				UE_CLOG(!bDefsApiVersionOk, LogTemp, Warning, TEXT("EOS_Achievements_DefinitionV2 version mismatch Expected=%d Actual=%d"), ExpectedDefinitionApiVersion, EOSDefinition->ApiVersion);
 
 				FString AchievementId = UTF8_TO_TCHAR(EOSDefinition->AchievementId);
 				FAchievementDefinition& AchievementDefinition = NewAchievementDefinitions.Emplace(AchievementId);
@@ -152,9 +153,10 @@ TOnlineAsyncOpHandle<FQueryAchievementDefinitions> FAchievementsEOSGS::QueryAchi
 				{
 					const EOS_Achievements_StatThresholds& EOSStatThreshold = EOSDefinition->StatThresholds[StatThresholdsIdx];
 
-					static_assert(EOS_ACHIEVEMENTS_STATTHRESHOLDS_API_LATEST == 1, "EOS_Achievements_StatThresholds updated, check new fields");
-					const bool bStatsApiVersionOk = EOSStatThreshold.ApiVersion == EOS_ACHIEVEMENTS_STATTHRESHOLDS_API_LATEST;
-					UE_CLOG(!bStatsApiVersionOk, LogTemp, Warning, TEXT("EOS_Achievements_StatThresholds version mismatch Expected=%d Actual=%d"), EOS_ACHIEVEMENTS_STATTHRESHOLDS_API_LATEST, EOSStatThreshold.ApiVersion);
+					const int ExpectedStatThresholdsApiVersion = 1;
+					UE_EOS_CHECK_API_MISMATCH(EOS_ACHIEVEMENTS_STATTHRESHOLDS_API_LATEST, ExpectedStatThresholdsApiVersion);
+					const bool bStatsApiVersionOk = EOSStatThreshold.ApiVersion == ExpectedStatThresholdsApiVersion;
+					UE_CLOG(!bStatsApiVersionOk, LogTemp, Warning, TEXT("EOS_Achievements_StatThresholds version mismatch Expected=%d Actual=%d"), ExpectedStatThresholdsApiVersion, EOSStatThreshold.ApiVersion);
 
 					FAchievementStatDefinition& StatDefinition = AchievementDefinition.StatDefinitions.Emplace_GetRef();
 					StatDefinition.StatId = UTF8_TO_TCHAR(EOSStatThreshold.Name);
@@ -238,8 +240,8 @@ TOnlineAsyncOpHandle<FQueryAchievementStates> FAchievementsEOSGS::QueryAchieveme
 			const FQueryAchievementStates::Params& Params = InAsyncOp.GetParams();
 
 			EOS_Achievements_QueryPlayerAchievementsOptions Options = {};
-			Options.ApiVersion = EOS_ACHIEVEMENTS_QUERYPLAYERACHIEVEMENTS_API_LATEST;
-			static_assert(EOS_ACHIEVEMENTS_QUERYPLAYERACHIEVEMENTS_API_LATEST == 2, "EOS_Achievements_QueryPlayerAchievementsOptions updated, check new fields");
+			Options.ApiVersion = 2;
+			UE_EOS_CHECK_API_MISMATCH(EOS_ACHIEVEMENTS_QUERYPLAYERACHIEVEMENTS_API_LATEST, 2);
 			Options.LocalUserId = GetProductUserIdChecked(Params.LocalAccountId);
 			Options.TargetUserId = Options.LocalUserId;
 
@@ -258,8 +260,8 @@ TOnlineAsyncOpHandle<FQueryAchievementStates> FAchievementsEOSGS::QueryAchieveme
 			}
 
 			EOS_Achievements_GetPlayerAchievementCountOptions GetCountOptions = {};
-			GetCountOptions.ApiVersion = EOS_ACHIEVEMENTS_GETPLAYERACHIEVEMENTCOUNT_API_LATEST;
-			static_assert(EOS_ACHIEVEMENTS_GETPLAYERACHIEVEMENTCOUNT_API_LATEST == 1, "EOS_Achievements_GetPlayerAchievementCountOptions updated, check new fields");
+			GetCountOptions.ApiVersion = 1;
+			UE_EOS_CHECK_API_MISMATCH(EOS_ACHIEVEMENTS_GETPLAYERACHIEVEMENTCOUNT_API_LATEST, 1);
 			GetCountOptions.UserId = GetProductUserIdChecked(Params.LocalAccountId);
 
 			const uint32_t AchievementCount = EOS_Achievements_GetPlayerAchievementCount(AchievementsHandle, &GetCountOptions);
@@ -268,8 +270,8 @@ TOnlineAsyncOpHandle<FQueryAchievementStates> FAchievementsEOSGS::QueryAchieveme
 			for (uint32_t AchievementIdx = 0; AchievementIdx < AchievementCount; AchievementIdx++)
 			{
 				EOS_Achievements_CopyPlayerAchievementByIndexOptions CopyOptions = {};
-				CopyOptions.ApiVersion = EOS_ACHIEVEMENTS_COPYPLAYERACHIEVEMENTBYINDEX_API_LATEST;
-				static_assert(EOS_ACHIEVEMENTS_COPYPLAYERACHIEVEMENTBYINDEX_API_LATEST == 2, "EOS_Achievements_CopyPlayerAchievementByIndexOptions updated, check new fields");
+				CopyOptions.ApiVersion = 2;
+				UE_EOS_CHECK_API_MISMATCH(EOS_ACHIEVEMENTS_COPYPLAYERACHIEVEMENTBYINDEX_API_LATEST, 2);
 				CopyOptions.LocalUserId = GetProductUserIdChecked(Params.LocalAccountId);
 				CopyOptions.TargetUserId = CopyOptions.LocalUserId;
 				CopyOptions.AchievementIndex = AchievementIdx;
@@ -284,11 +286,10 @@ TOnlineAsyncOpHandle<FQueryAchievementStates> FAchievementsEOSGS::QueryAchieveme
 					return;
 				}
 
-				static_assert(EOS_ACHIEVEMENTS_PLAYERACHIEVEMENT_API_LATEST == 2, "EOS_Achievements_PlayerAchievement updated, check new fields");
-				if (!ensure(EOSPlayerAchievement->ApiVersion == EOS_ACHIEVEMENTS_PLAYERACHIEVEMENT_API_LATEST))
-				{
-					UE_LOG(LogTemp, Warning, TEXT("EOS_Achievements_PlayerAchievement version mismatch Expected=%d Actual=%d"), EOS_ACHIEVEMENTS_PLAYERACHIEVEMENT_API_LATEST, EOSPlayerAchievement->ApiVersion);
-				}
+				const int ExpectedPlayerAchievementApiVersion = 2;
+				UE_EOS_CHECK_API_MISMATCH(EOS_ACHIEVEMENTS_PLAYERACHIEVEMENT_API_LATEST, ExpectedPlayerAchievementApiVersion);
+				const bool bPlayerAchievementApiVersionOk = EOSPlayerAchievement->ApiVersion == ExpectedPlayerAchievementApiVersion;
+				UE_CLOG(!bPlayerAchievementApiVersionOk, LogTemp, Warning, TEXT("EOS_Achievements_PlayerAchievement version mismatch Expected=%d Actual=%d"), ExpectedPlayerAchievementApiVersion, EOSPlayerAchievement->ApiVersion);
 
 				FString AchievementId = UTF8_TO_TCHAR(EOSPlayerAchievement->AchievementId);
 				FAchievementState& AchievementState = NewAchievementStates.Emplace(AchievementId);
@@ -388,8 +389,8 @@ TOnlineAsyncOpHandle<FUnlockAchievements> FAchievementsEOSGS::UnlockAchievements
 			}
 
 			EOS_Achievements_UnlockAchievementsOptions Options = {};
-			Options.ApiVersion = EOS_ACHIEVEMENTS_UNLOCKACHIEVEMENTS_API_LATEST;
-			static_assert(EOS_ACHIEVEMENTS_UNLOCKACHIEVEMENTS_API_LATEST == 1, "EOS_Achievements_UnlockAchievementsOptions updated, check new fields");
+			Options.ApiVersion = 1;
+			UE_EOS_CHECK_API_MISMATCH(EOS_ACHIEVEMENTS_UNLOCKACHIEVEMENTS_API_LATEST, 1);
 			Options.UserId = GetProductUserIdChecked(Params.LocalAccountId);
 			Options.AchievementIds = AchievementIdPtrs.GetData();
 			Options.AchievementsCount = AchievementIdPtrs.Num();
diff --git a/Engine/Plugins/Online/OnlineServicesEOSGS/Source/Private/Online/AuthEOSGS.cpp b/Engine/Plugins/Online/OnlineServicesEOSGS/Source/Private/Online/AuthEOSGS.cpp
--- a/Engine/Plugins/Online/OnlineServicesEOSGS/Source/Private/Online/AuthEOSGS.cpp
+++ b/Engine/Plugins/Online/OnlineServicesEOSGS/Source/Private/Online/AuthEOSGS.cpp
@@ -151,19 +151,19 @@ FEOSConnectLoginOptions& FEOSConnectLoginOptions::operator=(FEOSConnectLoginOpti
 FEOSConnectLoginOptions::FEOSConnectLoginOptions()
 {
 	// EOS_Connect_LoginOptions init
-	static_assert(EOS_CONNECT_LOGIN_API_LATEST == 2, "EOS_Connect_LoginOptions updated, check new fields");
-	ApiVersion = EOS_CONNECT_LOGIN_API_LATEST;
+	ApiVersion = 2;
+	UE_EOS_CHECK_API_MISMATCH(EOS_CONNECT_LOGIN_API_LATEST, 2);
 	Credentials = &CredentialsData;
 	UserLoginInfo = nullptr;
 
 	// EOS_Connect_Credentials init
-	static_assert(EOS_CONNECT_CREDENTIALS_API_LATEST == 1, "EOS_Connect_Credentials updated, check new fields");
-	CredentialsData.ApiVersion = EOS_CONNECT_CREDENTIALS_API_LATEST;
+	CredentialsData.ApiVersion = 1;
+	UE_EOS_CHECK_API_MISMATCH(EOS_CONNECT_CREDENTIALS_API_LATEST, 1);
 	CredentialsData.Token = nullptr;
 
 	// EOS_Connect_UserLoginInfo init
-	static_assert(EOS_CONNECT_USERLOGININFO_API_LATEST == 1, "EOS_Connect_UserLoginInfo updated, check new fields");
-	UserLoginInfoData.ApiVersion = EOS_CONNECT_USERLOGININFO_API_LATEST;
+	UserLoginInfoData.ApiVersion = 1;
+	UE_EOS_CHECK_API_MISMATCH(EOS_CONNECT_USERLOGININFO_API_LATEST, 1);
 	UserLoginInfoData.DisplayName = DisplayNameUtf8;
 }
 
@@ -319,14 +319,14 @@ FEOSAuthLoginOptions& FEOSAuthLoginOptions::operator=(FEOSAuthLoginOptions&& Oth
 FEOSAuthLoginOptions::FEOSAuthLoginOptions()
 {
 	// EOS_Auth_LoginOptions init
-	static_assert(EOS_AUTH_LOGIN_API_LATEST == 2, "EOS_Auth_LoginOptions updated, check new fields");
-	ApiVersion = EOS_AUTH_LOGIN_API_LATEST;
+	UE_EOS_CHECK_API_MISMATCH(EOS_AUTH_LOGIN_API_LATEST, 2);
+	ApiVersion = 2;
 	Credentials = &CredentialsData;
 	ScopeFlags = EOS_EAuthScopeFlags::EOS_AS_NoFlags;
 
 	// EOS_Auth_Credentials init
-	static_assert(EOS_AUTH_CREDENTIALS_API_LATEST == 3, "EOS_Auth_Credentials updated, check new fields");
-	CredentialsData.ApiVersion = EOS_AUTH_CREDENTIALS_API_LATEST;
+	UE_EOS_CHECK_API_MISMATCH(EOS_AUTH_CREDENTIALS_API_LATEST, 3);
+	CredentialsData.ApiVersion = 3;
 	CredentialsData.Id = nullptr;
 	CredentialsData.Token = nullptr;
 	CredentialsData.Type = EOS_ELoginCredentialType::EOS_LCT_Password;
@@ -835,9 +835,9 @@ TOnlineAsyncOpHandle<FAuthLogout> FAuthEOSGS::Logout(FAuthLogout::Params&& Param
 			TFuture<void> Future = Promise.GetFuture();
 
 			EOS_Auth_DeletePersistentAuthOptions DeletePersistentAuthOptions = { 0 };
-			DeletePersistentAuthOptions.ApiVersion = EOS_AUTH_DELETEPERSISTENTAUTH_API_LATEST;
+			DeletePersistentAuthOptions.ApiVersion = 2;
 			DeletePersistentAuthOptions.RefreshToken = nullptr; // TODO: Is this needed?  Docs say it's needed for consoles
-			static_assert(EOS_AUTH_DELETEPERSISTENTAUTH_API_LATEST == 2, "EOS_Auth_DeletePersistentAuthOptions updated, check new fields");
+			UE_EOS_CHECK_API_MISMATCH(EOS_AUTH_DELETEPERSISTENTAUTH_API_LATEST, 2);
 
 			EOS_Async(EOS_Auth_DeletePersistentAuth, AuthHandle, DeletePersistentAuthOptions,
 			[Promise = MoveTemp(Promise)](const EOS_Auth_DeletePersistentAuthCallbackInfo* Data) mutable -> void
@@ -998,9 +998,9 @@ TFuture<TDefaultErrorResult<FAuthLoginEASImpl>> FAuthEOSGS::LoginEASImpl(const F
 		else if (Data->ResultCode == EOS_EResult::EOS_InvalidUser && Data->ContinuanceToken != nullptr)
 		{
 			EOS_Auth_LinkAccountOptions LinkAccountOptions = {};
-			LinkAccountOptions.ApiVersion = EOS_AUTH_LINKACCOUNT_API_LATEST;
+			LinkAccountOptions.ApiVersion = 1;
 			LinkAccountOptions.ContinuanceToken = Data->ContinuanceToken;
-			static_assert(EOS_AUTH_LINKACCOUNT_API_LATEST == 1, "EOS_Auth_LinkAccountOptions updated, check new fields");
+			UE_EOS_CHECK_API_MISMATCH(EOS_AUTH_LINKACCOUNT_API_LATEST, 1);
 
 			EOS_Async(EOS_Auth_LinkAccount, AuthHandle, LinkAccountOptions,
 			[Promise = MoveTemp(Promise)](const EOS_Auth_LinkAccountCallbackInfo* Data) mutable -> void
@@ -1030,8 +1030,8 @@ TFuture<TDefaultErrorResult<FAuthLoginEASImpl>> FAuthEOSGS::LoginEASImpl(const F
 			if (IsPersistentAuthLogin && bShouldRemoveCachedToken)
 			{
 				EOS_Auth_DeletePersistentAuthOptions DeletePersistentAuthOptions = {};
-				DeletePersistentAuthOptions.ApiVersion = EOS_AUTH_DELETEPERSISTENTAUTH_API_LATEST;
-				static_assert(EOS_AUTH_DELETEPERSISTENTAUTH_API_LATEST == 2, "EOS_Auth_DeletePersistentAuthOptions updated, check new fields");
+				DeletePersistentAuthOptions.ApiVersion = 2;
+				UE_EOS_CHECK_API_MISMATCH(EOS_AUTH_DELETEPERSISTENTAUTH_API_LATEST, 2);
 
 				EOS_Async(EOS_Auth_DeletePersistentAuth, AuthHandle, DeletePersistentAuthOptions,
 				[ResolvedError = MoveTemp(ResolvedError), Promise = MoveTemp(Promise)](const EOS_Auth_DeletePersistentAuthCallbackInfo* Data) mutable -> void
@@ -1061,9 +1061,9 @@ TFuture<TDefaultErrorResult<FAuthLogoutEASImpl>> FAuthEOSGS::LogoutEASImpl(const
 	TFuture<TDefaultErrorResult<FAuthLogoutEASImpl>> Future = Promise.GetFuture();
 
 	EOS_Auth_LogoutOptions LogoutOptions = {};
-	LogoutOptions.ApiVersion = EOS_AUTH_LOGOUT_API_LATEST;
+	LogoutOptions.ApiVersion = 1;
 	LogoutOptions.LocalUserId = LogoutParams.EpicAccountId;
-	static_assert(EOS_AUTH_LOGOUT_API_LATEST == 1, "EOS_Auth_LogoutOptions updated, check new fields");
+	UE_EOS_CHECK_API_MISMATCH(EOS_AUTH_LOGOUT_API_LATEST, 1);
 
 	EOS_Async(EOS_Auth_Logout, AuthHandle, LogoutOptions,
 	[Promise = MoveTemp(Promise)](const EOS_Auth_LogoutCallbackInfo* Data) mutable -> void
@@ -1086,9 +1086,9 @@ TFuture<TDefaultErrorResult<FAuthLogoutEASImpl>> FAuthEOSGS::LogoutEASImpl(const
 TDefaultErrorResult<FAuthGetExternalAuthTokenImpl> FAuthEOSGS::GetExternalAuthTokenImpl(const FAuthGetExternalAuthTokenImpl::Params& Params)
 {
 	EOS_Auth_CopyIdTokenOptions CopyIdTokenOptions = {};
-	CopyIdTokenOptions.ApiVersion = EOS_AUTH_COPYIDTOKEN_API_LATEST;
+	CopyIdTokenOptions.ApiVersion = 1;
 	CopyIdTokenOptions.AccountId = Params.EpicAccountId;
-	static_assert(EOS_AUTH_COPYIDTOKEN_API_LATEST == 1, "EOS_Auth_CopyIdTokenOptions updated, check new fields");
+	UE_EOS_CHECK_API_MISMATCH(EOS_AUTH_COPYIDTOKEN_API_LATEST, 1);
 
 	EOS_Auth_IdToken* IdToken = nullptr;
 	EOS_EResult Result = EOS_Auth_CopyIdToken(AuthHandle, &CopyIdTokenOptions, &IdToken);
@@ -1133,9 +1133,9 @@ TFuture<TDefaultErrorResult<FAuthLoginConnectImpl>> FAuthEOSGS::LoginConnectImpl
 		else if (Data->ResultCode == EOS_EResult::EOS_InvalidUser && Data->ContinuanceToken != nullptr)
 		{
 			EOS_Connect_CreateUserOptions ConnectCreateUserOptions = { };
-			ConnectCreateUserOptions.ApiVersion = EOS_CONNECT_CREATEUSER_API_LATEST;
+			ConnectCreateUserOptions.ApiVersion = 1;
 			ConnectCreateUserOptions.ContinuanceToken = Data->ContinuanceToken;
-			static_assert(EOS_CONNECT_CREATEUSER_API_LATEST == 1, "EOS_Connect_CreateUserOptions updated, check new fields");
+			UE_EOS_CHECK_API_MISMATCH(EOS_CONNECT_CREATEUSER_API_LATEST, 1);
 
 			EOS_Async(EOS_Connect_CreateUser, ConnectHandle, ConnectCreateUserOptions,
 			[Promise = MoveTemp(Promise)](const EOS_Connect_CreateUserCallbackInfo* Data) mutable -> void
@@ -1405,35 +1405,44 @@ TOnlineAsyncOpHandle<FAuthHandleEASLoginStatusChangedImpl> FAuthEOSGS::HandleEAS
 
 void FAuthEOSGS::RegisterHandlers()
 {
-	// Register for EOS connect connection status updates.
-	OnConnectLoginStatusChangedEOSEventRegistration = EOS_RegisterComponentEventHandler(
-		this,
-		ConnectHandle,
-		EOS_CONNECT_ADDNOTIFYLOGINSTATUSCHANGED_API_LATEST,
-		&EOS_Connect_AddNotifyLoginStatusChanged,
-		&EOS_Connect_RemoveNotifyLoginStatusChanged,
-		&FAuthEOSGS::OnConnectLoginStatusChanged);
-	static_assert(EOS_CONNECT_ADDNOTIFYLOGINSTATUSCHANGED_API_LATEST == 1, "EOS_Connect_AddNotifyLoginStatusChanged updated, check new fields");
-
-	// Notification of a pending external token expiration ~10 minutes.
-	OnConnectAuthNotifyExpirationEOSEventRegistration = EOS_RegisterComponentEventHandler(
-		this,
-		ConnectHandle,
-		EOS_CONNECT_ONAUTHEXPIRATIONCALLBACK_API_LATEST,
-		&EOS_Connect_AddNotifyAuthExpiration,
-		&EOS_Connect_RemoveNotifyAuthExpiration,
-		&FAuthEOSGS::OnConnectAuthNotifyExpiration);
-	static_assert(EOS_CONNECT_ONAUTHEXPIRATIONCALLBACK_API_LATEST == 1, "EOS_Connect_AddNotifyAuthExpiration updated, check new fields");
-
-	// Register for EAS connection status updates.
-	OnConnectAuthNotifyExpirationEOSEventRegistration = EOS_RegisterComponentEventHandler(
-		this,
-		AuthHandle,
-		EOS_AUTH_ADDNOTIFYLOGINSTATUSCHANGED_API_LATEST,
-		&EOS_Auth_AddNotifyLoginStatusChanged,
-		&EOS_Auth_RemoveNotifyLoginStatusChanged,
-		&FAuthEOSGS::OnEASLoginStatusChanged);
-	static_assert(EOS_AUTH_ADDNOTIFYLOGINSTATUSCHANGED_API_LATEST == 1, "EOS_Auth_AddNotifyLoginStatusChanged updated, check new fields");
+	{
+		// Register for EOS connect connection status updates.
+		const int ApiVersion = 1;
+		UE_EOS_CHECK_API_MISMATCH(EOS_CONNECT_ADDNOTIFYLOGINSTATUSCHANGED_API_LATEST, ApiVersion);
+		OnConnectLoginStatusChangedEOSEventRegistration = EOS_RegisterComponentEventHandler(
+			this,
+			ConnectHandle,
+			ApiVersion,
+			&EOS_Connect_AddNotifyLoginStatusChanged,
+			&EOS_Connect_RemoveNotifyLoginStatusChanged,
+			&FAuthEOSGS::OnConnectLoginStatusChanged);
+	}
+
+	{
+		// Notification of a pending external token expiration ~10 minutes.
+		const int ApiVersion = 1;
+		UE_EOS_CHECK_API_MISMATCH(EOS_CONNECT_ADDNOTIFYAUTHEXPIRATION_API_LATEST, ApiVersion);
+		OnConnectAuthNotifyExpirationEOSEventRegistration = EOS_RegisterComponentEventHandler(
+			this,
+			ConnectHandle,
+			ApiVersion,
+			&EOS_Connect_AddNotifyAuthExpiration,
+			&EOS_Connect_RemoveNotifyAuthExpiration,
+			&FAuthEOSGS::OnConnectAuthNotifyExpiration);
+	}
+
+	{
+		// Register for EAS connection status updates.
+		const int ApiVersion = 1;
+		UE_EOS_CHECK_API_MISMATCH(EOS_AUTH_ADDNOTIFYLOGINSTATUSCHANGED_API_LATEST, ApiVersion);
+		OnConnectAuthNotifyExpirationEOSEventRegistration = EOS_RegisterComponentEventHandler(
+			this,
+			AuthHandle,
+			ApiVersion,
+			&EOS_Auth_AddNotifyLoginStatusChanged,
+			&EOS_Auth_RemoveNotifyLoginStatusChanged,
+			&FAuthEOSGS::OnEASLoginStatusChanged);
+	}
 }
 
 void FAuthEOSGS::UnregisterHandlers()
diff --git a/Engine/Plugins/Online/OnlineServicesEOSGS/Source/Private/Online/LeaderboardsEOSGS.cpp b/Engine/Plugins/Online/OnlineServicesEOSGS/Source/Private/Online/LeaderboardsEOSGS.cpp
--- a/Engine/Plugins/Online/OnlineServicesEOSGS/Source/Private/Online/LeaderboardsEOSGS.cpp
+++ b/Engine/Plugins/Online/OnlineServicesEOSGS/Source/Private/Online/LeaderboardsEOSGS.cpp
@@ -53,8 +53,8 @@ TOnlineAsyncOpHandle<FReadEntriesForUsers> FLeaderboardsEOSGS::ReadEntriesForUse
 			}
 
 			EOS_Leaderboards_UserScoresQueryStatInfo StatInfo;
-			StatInfo.ApiVersion = EOS_LEADERBOARDS_USERSCORESQUERYSTATINFO_API_LATEST;
-			static_assert(EOS_LEADERBOARDS_USERSCORESQUERYSTATINFO_API_LATEST == 1, "EOS_Leaderboards_UserScoresQueryStatInfo updated, check new fields");
+			StatInfo.ApiVersion = 1;
+			UE_EOS_CHECK_API_MISMATCH(EOS_LEADERBOARDS_USERSCORESQUERYSTATINFO_API_LATEST, 1);
 
 			char StatNameANSI[EOS_LEADERBOARD_STAT_NAME_MAX_LENGTH];
 			FCStringAnsi::Strncpy(StatNameANSI, TCHAR_TO_UTF8(*Params.BoardName), EOS_LEADERBOARD_STAT_NAME_MAX_LENGTH);
@@ -63,8 +63,8 @@ TOnlineAsyncOpHandle<FReadEntriesForUsers> FLeaderboardsEOSGS::ReadEntriesForUse
 			StatInfo.Aggregation = EOS_ELeaderboardAggregation::EOS_LA_Latest; // TODO: Use Stats definitions
 
 			EOS_Leaderboards_QueryLeaderboardUserScoresOptions Options = { };
-			Options.ApiVersion = EOS_LEADERBOARDS_QUERYLEADERBOARDUSERSCORES_API_LATEST;
-			static_assert(EOS_LEADERBOARDS_QUERYLEADERBOARDUSERSCORES_API_LATEST == 2, "EOS_Leaderboards_QueryLeaderboardUserScores updated, check new fields");
+			Options.ApiVersion = 2;
+			UE_EOS_CHECK_API_MISMATCH(EOS_LEADERBOARDS_QUERYLEADERBOARDUSERSCORES_API_LATEST, 2);
 			Options.UserIds = ProductUserIds.GetData();
 			Options.UserIdsCount = ProductUserIds.Num();
 			Options.StatInfo = &StatInfo;
@@ -94,8 +94,8 @@ TOnlineAsyncOpHandle<FReadEntriesForUsers> FLeaderboardsEOSGS::ReadEntriesForUse
 			for (const FAccountId& AccountId : Params.AccountIds)
 			{
 				EOS_Leaderboards_CopyLeaderboardUserScoreByUserIdOptions UserCopyOptions = { };
-				UserCopyOptions.ApiVersion = EOS_LEADERBOARDS_COPYLEADERBOARDUSERSCOREBYUSERID_API_LATEST;
-				static_assert(EOS_LEADERBOARDS_COPYLEADERBOARDUSERSCOREBYUSERID_API_LATEST == 1, "EOS_Leaderboards_CopyLeaderboardUserScoreByUserIdOptions updated, check new fields");
+				UserCopyOptions.ApiVersion = 1;
+				UE_EOS_CHECK_API_MISMATCH(EOS_LEADERBOARDS_COPYLEADERBOARDUSERSCOREBYUSERID_API_LATEST, 1);
 
 				UserCopyOptions.UserId = GetProductUserIdChecked(AccountId);
 				UserCopyOptions.StatName = StatNameANSI;
@@ -131,8 +131,8 @@ void QueryLeaderboardsEOS(EOS_HLeaderboards LeaderboardsHandle, const FAccountId
 {
 	EOS_Leaderboards_QueryLeaderboardRanksOptions Options;
 	Options.LocalUserId = GetProductUserIdChecked(LocalAccountId);
-	Options.ApiVersion = EOS_LEADERBOARDS_QUERYLEADERBOARDRANKS_API_LATEST;
-	static_assert(EOS_LEADERBOARDS_QUERYLEADERBOARDRANKS_API_LATEST == 2, "EOS_Leaderboards_QueryLeaderboardRanks updated, check new fields");
+	Options.ApiVersion = 2;
+	UE_EOS_CHECK_API_MISMATCH(EOS_LEADERBOARDS_QUERYLEADERBOARDRANKS_API_LATEST, 2);
 
 	char LeaderboardNameANSI[EOS_LEADERBOARD_STAT_NAME_MAX_LENGTH];
 	FCStringAnsi::Strncpy(LeaderboardNameANSI, TCHAR_TO_UTF8(*BoardName), EOS_LEADERBOARD_STAT_NAME_MAX_LENGTH);
@@ -146,8 +146,8 @@ void ReadEntriesInRange(EOS_HLeaderboards LeaderboardsHandle, uint32 StartIndex,
 	for (uint32 i = StartIndex; i <= EndIndex; ++i)
 	{
 		EOS_Leaderboards_CopyLeaderboardRecordByIndexOptions CopyRecordOptions;
-		CopyRecordOptions.ApiVersion = EOS_LEADERBOARDS_COPYLEADERBOARDRECORDBYINDEX_API_LATEST;
-		static_assert(EOS_LEADERBOARDS_COPYLEADERBOARDRECORDBYINDEX_API_LATEST == 2, "EOS_Leaderboards_CopyLeaderboardRecordByIndexOptions updated, check new fields");
+		CopyRecordOptions.ApiVersion = 2;
+		UE_EOS_CHECK_API_MISMATCH(EOS_LEADERBOARDS_COPYLEADERBOARDRECORDBYINDEX_API_LATEST, 2);
 		CopyRecordOptions.LeaderboardRecordIndex = i;
 
 		EOS_Leaderboards_LeaderboardRecord* LeaderboardRecord = nullptr;
@@ -248,8 +248,8 @@ TOnlineAsyncOpHandle<FReadEntriesAroundUser> FLeaderboardsEOSGS::ReadEntriesArou
 
 			EOS_Leaderboards_LeaderboardRecord* LeaderboardRecord = nullptr;
 			EOS_Leaderboards_CopyLeaderboardRecordByUserIdOptions CopyRecordByUserIdOptions;
-			CopyRecordByUserIdOptions.ApiVersion = EOS_LEADERBOARDS_COPYLEADERBOARDRECORDBYUSERID_API_LATEST;
-			static_assert(EOS_LEADERBOARDS_COPYLEADERBOARDRECORDBYUSERID_API_LATEST == 2, "EOS_Leaderboards_CopyLeaderboardRecordByUserIdOptions updated, check new fields");
+			CopyRecordByUserIdOptions.ApiVersion = 2;
+			UE_EOS_CHECK_API_MISMATCH(EOS_LEADERBOARDS_COPYLEADERBOARDRECORDBYUSERID_API_LATEST, 2);
 			CopyRecordByUserIdOptions.UserId = GetProductUserIdChecked(Params.AccountId);
 			EOS_EResult CopyResult = EOS_Leaderboards_CopyLeaderboardRecordByUserId(LeaderboardsHandle, &CopyRecordByUserIdOptions, &LeaderboardRecord);
 			if (CopyResult == EOS_EResult::EOS_Success)
diff --git a/Engine/Plugins/Online/OnlineServicesEOSGS/Source/Private/Online/LobbiesEOSGS.cpp b/Engine/Plugins/Online/OnlineServicesEOSGS/Source/Private/Online/LobbiesEOSGS.cpp
--- a/Engine/Plugins/Online/OnlineServicesEOSGS/Source/Private/Online/LobbiesEOSGS.cpp
+++ b/Engine/Plugins/Online/OnlineServicesEOSGS/Source/Private/Online/LobbiesEOSGS.cpp
@@ -126,7 +126,8 @@ TOnlineAsyncOpHandle<FCreateLobby> FLobbiesEOSGS::CreateLobby(FCreateLobby::Para
 		// Once all local members are joined and the lobby attributes have been set the privacy setting will be moved to the user setting.
 
 		EOS_Lobby_CreateLobbyOptions CreateLobbyOptions = {};
-		CreateLobbyOptions.ApiVersion = EOS_LOBBY_CREATELOBBY_API_LATEST;
+		CreateLobbyOptions.ApiVersion = 8;
+		UE_EOS_CHECK_API_MISMATCH(EOS_LOBBY_CREATELOBBY_API_LATEST, 8);
 		CreateLobbyOptions.LocalUserId = GetProductUserIdChecked(Params.LocalAccountId);
 		CreateLobbyOptions.MaxLobbyMembers = Params.MaxMembers;
 		// Prevent lobby from appearing in search results until fully created with all user attributes.
@@ -136,7 +137,6 @@ TOnlineAsyncOpHandle<FCreateLobby> FLobbiesEOSGS::CreateLobby(FCreateLobby::Para
 		CreateLobbyOptions.BucketId = BucketTanslator.GetBucketIdEOS();
 		CreateLobbyOptions.bDisableHostMigration = false; // todo: handle
 		CreateLobbyOptions.bEnableRTCRoom = false; // todo: handle
-		static_assert(EOS_LOBBY_CREATELOBBY_API_LATEST == 8, "EOS_Lobby_CreateLobbyOptions updated, check new fields");
 
 		EOS_Async(EOS_Lobby_CreateLobby, LobbyPrerequisites->LobbyInterfaceHandle, CreateLobbyOptions, MoveTemp(Promise)); 
 	})
@@ -383,12 +383,12 @@ TOnlineAsyncOpHandle<FJoinLobby> FLobbiesEOSGS::JoinLobby(FJoinLobby::Params&& I
 		const TSharedRef<FLobbyDetailsEOS>& LobbyDetails = GetOpDataChecked<TSharedRef<FLobbyDetailsEOS>>(InAsyncOp, LobbyDetailsKeyName);
 
 		EOS_Lobby_JoinLobbyOptions JoinLobbyOptions = {};
-		JoinLobbyOptions.ApiVersion = EOS_LOBBY_JOINLOBBY_API_LATEST;
+		JoinLobbyOptions.ApiVersion = 3;
+		UE_EOS_CHECK_API_MISMATCH(EOS_LOBBY_JOINLOBBY_API_LATEST, 3);
 		JoinLobbyOptions.LobbyDetailsHandle = LobbyDetails->GetEOSHandle();
 		JoinLobbyOptions.LocalUserId = GetProductUserIdChecked(Params.LocalAccountId);
 		JoinLobbyOptions.bPresenceEnabled = false; // todo
 		JoinLobbyOptions.LocalRTCOptions = nullptr;
-		static_assert(EOS_LOBBY_JOINLOBBY_API_LATEST == 3, "EOS_Lobby_JoinLobbyOptions updated, check new fields");
 
 		EOS_Async(EOS_Lobby_JoinLobby, LobbyPrerequisites->LobbyInterfaceHandle, JoinLobbyOptions, MoveTemp(Promise));
 	})
@@ -1401,61 +1401,61 @@ void FLobbiesEOSGS::RegisterHandlers()
 	OnLobbyUpdatedEOSEventRegistration = EOS_RegisterComponentEventHandler(
 		this,
 		LobbyPrerequisites->LobbyInterfaceHandle,
-		EOS_LOBBY_ADDNOTIFYLOBBYUPDATERECEIVED_API_LATEST,
+		1,
 		&EOS_Lobby_AddNotifyLobbyUpdateReceived,
 		&EOS_Lobby_RemoveNotifyLobbyUpdateReceived,
 		&FLobbiesEOSGS::HandleLobbyUpdated);
-	static_assert(EOS_LOBBY_ADDNOTIFYLOBBYUPDATERECEIVED_API_LATEST == 1, "EOS_Lobby_AddNotifyLobbyUpdateReceived updated, check new fields");
+	UE_EOS_CHECK_API_MISMATCH(EOS_LOBBY_ADDNOTIFYLOBBYUPDATERECEIVED_API_LATEST, 1);
 
 	// Register for lobby member updates.
 	OnLobbyMemberUpdatedEOSEventRegistration = EOS_RegisterComponentEventHandler(
 		this,
 		LobbyPrerequisites->LobbyInterfaceHandle,
-		EOS_LOBBY_ADDNOTIFYLOBBYMEMBERUPDATERECEIVED_API_LATEST,
+		1,
 		&EOS_Lobby_AddNotifyLobbyMemberUpdateReceived,
 		&EOS_Lobby_RemoveNotifyLobbyMemberUpdateReceived,
 		&FLobbiesEOSGS::HandleLobbyMemberUpdated);
-	static_assert(EOS_LOBBY_ADDNOTIFYLOBBYMEMBERUPDATERECEIVED_API_LATEST == 1, "EOS_Lobby_AddNotifyLobbyMemberUpdateReceived updated, check new fields");
+	UE_EOS_CHECK_API_MISMATCH(EOS_LOBBY_ADDNOTIFYLOBBYMEMBERUPDATERECEIVED_API_LATEST, 1);
 
 	// Register for lobby member status changed.
 	OnLobbyMemberStatusReceivedEOSEventRegistration = EOS_RegisterComponentEventHandler(
 		this,
 		LobbyPrerequisites->LobbyInterfaceHandle,
-		EOS_LOBBY_ADDNOTIFYLOBBYMEMBERSTATUSRECEIVED_API_LATEST,
+		1,
 		&EOS_Lobby_AddNotifyLobbyMemberStatusReceived,
 		&EOS_Lobby_RemoveNotifyLobbyMemberStatusReceived,
 		&FLobbiesEOSGS::HandleLobbyMemberStatusReceived);
-	static_assert(EOS_LOBBY_ADDNOTIFYLOBBYMEMBERSTATUSRECEIVED_API_LATEST == 1, "EOS_Lobby_AddNotifyLobbyMemberStatusReceived updated, check new fields");
+	UE_EOS_CHECK_API_MISMATCH(EOS_LOBBY_ADDNOTIFYLOBBYMEMBERSTATUSRECEIVED_API_LATEST, 1);
 
 	// Register for lobby invite received.
 	OnLobbyInviteReceivedEOSEventRegistration = EOS_RegisterComponentEventHandler(
 		this,
 		LobbyPrerequisites->LobbyInterfaceHandle,
-		EOS_LOBBY_ADDNOTIFYLOBBYINVITERECEIVED_API_LATEST,
+		1,
 		&EOS_Lobby_AddNotifyLobbyInviteReceived,
 		&EOS_Lobby_RemoveNotifyLobbyInviteReceived,
 		&FLobbiesEOSGS::HandleLobbyInviteReceived);
-	static_assert(EOS_LOBBY_ADDNOTIFYLOBBYINVITERECEIVED_API_LATEST == 1, "EOS_Lobby_AddNotifyLobbyInviteReceived updated, check new fields");
+	UE_EOS_CHECK_API_MISMATCH(EOS_LOBBY_ADDNOTIFYLOBBYINVITERECEIVED_API_LATEST, 1);
 
 	// Register for lobby invite accepted.
 	OnLobbyInviteAcceptedEOSEventRegistration = EOS_RegisterComponentEventHandler(
 		this,
 		LobbyPrerequisites->LobbyInterfaceHandle,
-		EOS_LOBBY_ADDNOTIFYLOBBYINVITEACCEPTED_API_LATEST,
+		1,
 		&EOS_Lobby_AddNotifyLobbyInviteAccepted,
 		&EOS_Lobby_RemoveNotifyLobbyInviteAccepted,
 		&FLobbiesEOSGS::HandleLobbyInviteAccepted);
-	static_assert(EOS_LOBBY_ADDNOTIFYLOBBYINVITEACCEPTED_API_LATEST == 1, "EOS_Lobby_AddNotifyLobbyInviteAccepted updated, check new fields");
+	UE_EOS_CHECK_API_MISMATCH(EOS_LOBBY_ADDNOTIFYLOBBYINVITEACCEPTED_API_LATEST, 1);
 
 	// Register for join lobby accepted via overlay.
 	OnJoinLobbyAcceptedEOSEventRegistration = EOS_RegisterComponentEventHandler(
 		this,
 		LobbyPrerequisites->LobbyInterfaceHandle,
-		EOS_LOBBY_ADDNOTIFYJOINLOBBYACCEPTED_API_LATEST,
+		1,
 		&EOS_Lobby_AddNotifyJoinLobbyAccepted,
 		&EOS_Lobby_RemoveNotifyJoinLobbyAccepted,
 		&FLobbiesEOSGS::HandleJoinLobbyAccepted);
-	static_assert(EOS_LOBBY_ADDNOTIFYJOINLOBBYACCEPTED_API_LATEST == 1, "EOS_Lobby_AddNotifyJoinLobbyAccepted updated, check new fields");
+	UE_EOS_CHECK_API_MISMATCH(EOS_LOBBY_ADDNOTIFYJOINLOBBYACCEPTED_API_LATEST, 1);
 }
 
 void FLobbiesEOSGS::UnregisterHandlers()
@@ -1536,10 +1536,10 @@ TFuture<TDefaultErrorResult<FLobbiesLeaveLobbyImpl>> FLobbiesEOSGS::LeaveLobbyIm
 
 	// Start operation.
 	EOS_Lobby_LeaveLobbyOptions LeaveLobbyOptions = {};
-	LeaveLobbyOptions.ApiVersion = EOS_LOBBY_LEAVELOBBY_API_LATEST;
+	LeaveLobbyOptions.ApiVersion = 1;
+	UE_EOS_CHECK_API_MISMATCH(EOS_LOBBY_LEAVELOBBY_API_LATEST, 1);
 	LeaveLobbyOptions.LocalUserId = GetProductUserIdChecked(Params.LocalAccountId);
 	LeaveLobbyOptions.LobbyId = Params.LobbyData->GetLobbyIdEOS();
-	static_assert(EOS_LOBBY_LEAVELOBBY_API_LATEST == 1, "EOS_Lobby_LeaveLobbyOptions updated, check new fields");
 
 	TPromise<TDefaultErrorResult<FLobbiesLeaveLobbyImpl>> Promise;
 	TFuture<TDefaultErrorResult<FLobbiesLeaveLobbyImpl>> Future = Promise.GetFuture();
@@ -1585,10 +1585,10 @@ TFuture<TDefaultErrorResult<FLobbiesDestroyLobbyImpl>> FLobbiesEOSGS::DestroyLob
 	// Start operation.
 	FTCHARToUTF8 LobbyIdTranslator(*Params.LobbyIdString);
 	EOS_Lobby_DestroyLobbyOptions DestroyLobbyOptions = {};
-	DestroyLobbyOptions.ApiVersion = EOS_LOBBY_DESTROYLOBBY_API_LATEST;
+	DestroyLobbyOptions.ApiVersion = 1;
+	UE_EOS_CHECK_API_MISMATCH(EOS_LOBBY_DESTROYLOBBY_API_LATEST, 1);
 	DestroyLobbyOptions.LocalUserId = GetProductUserIdChecked(Params.LocalAccountId);
 	DestroyLobbyOptions.LobbyId = LobbyIdTranslator.Get();
-	static_assert(EOS_LOBBY_DESTROYLOBBY_API_LATEST == 1, "EOS_Lobby_DestroyLobbyOptions updated, check new fields");
 
 	TPromise<TDefaultErrorResult<FLobbiesDestroyLobbyImpl>> Promise;
 	TFuture<TDefaultErrorResult<FLobbiesDestroyLobbyImpl>> Future = Promise.GetFuture();
@@ -1633,11 +1633,11 @@ TFuture<TDefaultErrorResult<FLobbiesInviteLobbyMemberImpl>> FLobbiesEOSGS::Invit
 
 	// Start operation.
 	EOS_Lobby_SendInviteOptions SendInviteOptions = {};
-	SendInviteOptions.ApiVersion = EOS_LOBBY_SENDINVITE_API_LATEST;
+	SendInviteOptions.ApiVersion = 1;
+	UE_EOS_CHECK_API_MISMATCH(EOS_LOBBY_SENDINVITE_API_LATEST, 1);
 	SendInviteOptions.LocalUserId = GetProductUserIdChecked(Params.LocalAccountId);
 	SendInviteOptions.TargetUserId = GetProductUserIdChecked(Params.TargetAccountId);
 	SendInviteOptions.LobbyId = Params.LobbyData->GetLobbyIdEOS();
-	static_assert(EOS_LOBBY_SENDINVITE_API_LATEST == 1, "EOS_Lobby_SendInviteOptions updated, check new fields");
 
 	TPromise<TDefaultErrorResult<FLobbiesInviteLobbyMemberImpl>> Promise;
 	TFuture<TDefaultErrorResult<FLobbiesInviteLobbyMemberImpl>> Future = Promise.GetFuture();
@@ -1692,10 +1692,10 @@ TFuture<TDefaultErrorResult<FLobbiesDeclineLobbyInvitationImpl>> FLobbiesEOSGS::
 
 	// Start operation.
 	EOS_Lobby_RejectInviteOptions RejectInviteOptions = {};
-	RejectInviteOptions.ApiVersion = EOS_LOBBY_REJECTINVITE_API_LATEST;
+	RejectInviteOptions.ApiVersion = 1;
+	UE_EOS_CHECK_API_MISMATCH(EOS_LOBBY_REJECTINVITE_API_LATEST, 1);
 	RejectInviteOptions.LocalUserId = GetProductUserIdChecked(Params.LocalAccountId);
 	RejectInviteOptions.InviteId = InviteData->GetInviteIdEOS();
-	static_assert(EOS_LOBBY_REJECTINVITE_API_LATEST == 1, "EOS_Lobby_RejectInviteOptions updated, check new fields");
 
 	TPromise<TDefaultErrorResult<FLobbiesDeclineLobbyInvitationImpl>> Promise;
 	TFuture<TDefaultErrorResult<FLobbiesDeclineLobbyInvitationImpl>> Future = Promise.GetFuture();
@@ -1751,11 +1751,11 @@ TFuture<TDefaultErrorResult<FLobbiesKickLobbyMemberImpl>> FLobbiesEOSGS::KickLob
 
 	// Start operation.
 	EOS_Lobby_KickMemberOptions KickMemberOptions = {};
-	KickMemberOptions.ApiVersion = EOS_LOBBY_KICKMEMBER_API_LATEST;
+	KickMemberOptions.ApiVersion = 1;
+	UE_EOS_CHECK_API_MISMATCH(EOS_LOBBY_KICKMEMBER_API_LATEST, 1);
 	KickMemberOptions.LocalUserId = GetProductUserIdChecked(Params.LocalAccountId);
 	KickMemberOptions.TargetUserId = GetProductUserIdChecked(Params.TargetAccountId);
 	KickMemberOptions.LobbyId = Params.LobbyData->GetLobbyIdEOS();
-	static_assert(EOS_LOBBY_KICKMEMBER_API_LATEST == 1, "EOS_Lobby_KickMemberOptions updated, check new fields");
 
 	TPromise<TDefaultErrorResult<FLobbiesKickLobbyMemberImpl>> Promise;
 	TFuture<TDefaultErrorResult<FLobbiesKickLobbyMemberImpl>> Future = Promise.GetFuture();
@@ -1808,11 +1808,11 @@ TFuture<TDefaultErrorResult<FLobbiesPromoteLobbyMemberImpl>> FLobbiesEOSGS::Prom
 
 	// Start operation.
 	EOS_Lobby_PromoteMemberOptions PromoteMemberOptions = {};
-	PromoteMemberOptions.ApiVersion = EOS_LOBBY_PROMOTEMEMBER_API_LATEST;
+	PromoteMemberOptions.ApiVersion = 1;
+	UE_EOS_CHECK_API_MISMATCH(EOS_LOBBY_PROMOTEMEMBER_API_LATEST, 1);
 	PromoteMemberOptions.LocalUserId = GetProductUserIdChecked(Params.LocalAccountId);
 	PromoteMemberOptions.TargetUserId = GetProductUserIdChecked(Params.TargetAccountId);
 	PromoteMemberOptions.LobbyId = Params.LobbyData->GetLobbyIdEOS();
-	static_assert(EOS_LOBBY_PROMOTEMEMBER_API_LATEST == 1, "EOS_Lobby_PromoteMemberOptions updated, check new fields");
 
 	TPromise<TDefaultErrorResult<FLobbiesPromoteLobbyMemberImpl>> Promise;
 	TFuture<TDefaultErrorResult<FLobbiesPromoteLobbyMemberImpl>> Future = Promise.GetFuture();
diff --git a/Engine/Plugins/Online/OnlineServicesEOSGS/Source/Private/Online/LobbiesEOSGSTypes.cpp b/Engine/Plugins/Online/OnlineServicesEOSGS/Source/Private/Online/LobbiesEOSGSTypes.cpp
--- a/Engine/Plugins/Online/OnlineServicesEOSGS/Source/Private/Online/LobbiesEOSGSTypes.cpp
+++ b/Engine/Plugins/Online/OnlineServicesEOSGS/Source/Private/Online/LobbiesEOSGSTypes.cpp
@@ -51,9 +51,9 @@ FLobbyAttributeTranslator<ELobbyTranslationType::ToService>::FLobbyAttributeTran
 	: KeyConverterStorage(*FromAttributeData.Id.ToString().ToLower())
 	, AttributeVisibility(EOS_ELobbyAttributeVisibility::EOS_LAT_PRIVATE)
 {
-	AttributeData.ApiVersion = EOS_LOBBY_ATTRIBUTEDATA_API_LATEST;
+	AttributeData.ApiVersion = 1;
+	UE_EOS_CHECK_API_MISMATCH(EOS_LOBBY_ATTRIBUTEDATA_API_LATEST, 1);
 	AttributeData.Key = KeyConverterStorage.Get();
-	static_assert(EOS_LOBBY_ATTRIBUTEDATA_API_LATEST == 1, "EOS_Lobby_AttributeData updated, check new fields");
 
 	switch (FromAttributeData.Value.GetType())
 	{
@@ -150,10 +150,10 @@ TDefaultErrorResultInternal<TSharedRef<FLobbyDetailsEOS>> FLobbyDetailsEOS::Crea
 	EOS_HLobbyDetails LobbyDetailsHandle = {};
 
 	EOS_Lobby_CopyLobbyDetailsHandleOptions Options;
-	Options.ApiVersion = EOS_LOBBY_COPYLOBBYDETAILSHANDLE_API_LATEST;
+	Options.ApiVersion = 1;
+	UE_EOS_CHECK_API_MISMATCH(EOS_LOBBY_COPYLOBBYDETAILSHANDLE_API_LATEST, 1);
 	Options.LobbyId = LobbyIdEOS;
 	Options.LocalUserId = GetProductUserIdChecked(LocalAccountId);
-	static_assert(EOS_LOBBY_COPYLOBBYDETAILSHANDLE_API_LATEST == 1, "EOS_Lobby_CopyLobbyDetailsHandleOptions updated, check new fields");
 
 	EOS_EResult EOSResult = EOS_Lobby_CopyLobbyDetailsHandle(Prerequisites->LobbyInterfaceHandle, &Options, &LobbyDetailsHandle);
 	if (EOSResult != EOS_EResult::EOS_Success)
@@ -193,9 +193,9 @@ TDefaultErrorResultInternal<TSharedRef<FLobbyDetailsEOS>> FLobbyDetailsEOS::Crea
 	EOS_HLobbyDetails LobbyDetailsHandle = {};
 
 	EOS_Lobby_CopyLobbyDetailsHandleByInviteIdOptions Options;
-	Options.ApiVersion = EOS_LOBBY_COPYLOBBYDETAILSHANDLEBYINVITEID_API_LATEST;
+	Options.ApiVersion = 1;
+	UE_EOS_CHECK_API_MISMATCH(EOS_LOBBY_COPYLOBBYDETAILSHANDLEBYINVITEID_API_LATEST, 1);
 	Options.InviteId = InviteId;
-	static_assert(EOS_LOBBY_COPYLOBBYDETAILSHANDLEBYINVITEID_API_LATEST == 1, "EOS_Lobby_CopyLobbyDetailsHandleByInviteIdOptions updated, check new fields");
 
 	EOS_EResult EOSResult = EOS_Lobby_CopyLobbyDetailsHandleByInviteId(Prerequisites->LobbyInterfaceHandle, &Options, &LobbyDetailsHandle);
 	if (EOSResult != EOS_EResult::EOS_Success)
@@ -227,9 +227,9 @@ TDefaultErrorResultInternal<TSharedRef<FLobbyDetailsEOS>> FLobbyDetailsEOS::Crea
 	EOS_HLobbyDetails LobbyDetailsHandle = {};
 
 	EOS_Lobby_CopyLobbyDetailsHandleByUiEventIdOptions Options;
-	Options.ApiVersion = EOS_LOBBY_COPYLOBBYDETAILSHANDLEBYUIEVENTID_API_LATEST;
+	Options.ApiVersion = 1;
+	UE_EOS_CHECK_API_MISMATCH(EOS_LOBBY_COPYLOBBYDETAILSHANDLEBYUIEVENTID_API_LATEST, 1);
 	Options.UiEventId = UiEventId;
-	static_assert(EOS_LOBBY_COPYLOBBYDETAILSHANDLEBYUIEVENTID_API_LATEST == 1, "EOS_Lobby_CopyLobbyDetailsHandleByUiEventIdOptions updated, check new fields");
 
 	EOS_EResult EOSResult = EOS_Lobby_CopyLobbyDetailsHandleByUiEventId(Prerequisites->LobbyInterfaceHandle, &Options, &LobbyDetailsHandle);
 	if (EOSResult != EOS_EResult::EOS_Success)
@@ -262,9 +262,9 @@ TDefaultErrorResultInternal<TSharedRef<FLobbyDetailsEOS>> FLobbyDetailsEOS::Crea
 	EOS_HLobbyDetails LobbyDetailsHandle = {};
 
 	EOS_LobbySearch_CopySearchResultByIndexOptions Options;
-	Options.ApiVersion = EOS_LOBBYSEARCH_COPYSEARCHRESULTBYINDEX_API_LATEST;
+	Options.ApiVersion = 1;
+	UE_EOS_CHECK_API_MISMATCH(EOS_LOBBYSEARCH_COPYSEARCHRESULTBYINDEX_API_LATEST, 1);
 	Options.LobbyIndex = ResultIndex;
-	static_assert(EOS_LOBBYSEARCH_COPYSEARCHRESULTBYINDEX_API_LATEST == 1, "EOS_LobbySearch_CopySearchResultByIndexOptions updated, check new fields");
 
 	EOS_EResult EOSResult = EOS_LobbySearch_CopySearchResultByIndex(SearchHandle, &Options, &LobbyDetailsHandle);
 	if (EOSResult != EOS_EResult::EOS_Success)
@@ -309,8 +309,8 @@ TFuture<TDefaultErrorResultInternal<FLobbyServiceSnapshot>> FLobbyDetailsEOS::Ge
 	}
 
 	EOS_LobbyDetails_GetMemberCountOptions GetMemberCountOptions = {};
-	GetMemberCountOptions.ApiVersion = EOS_LOBBYDETAILS_GETMEMBERCOUNT_API_LATEST;
-	static_assert(EOS_LOBBYDETAILS_GETMEMBERCOUNT_API_LATEST == 1, "EOS_LobbyDetails_GetMemberCountOptions updated, check new fields");
+	GetMemberCountOptions.ApiVersion = 1;
+	UE_EOS_CHECK_API_MISMATCH(EOS_LOBBYDETAILS_GETMEMBERCOUNT_API_LATEST, 1);
 
 	const uint32_t MemberCount = EOS_LobbyDetails_GetMemberCount(LobbyDetailsHandle, &GetMemberCountOptions);
 
@@ -320,9 +320,9 @@ TFuture<TDefaultErrorResultInternal<FLobbyServiceSnapshot>> FLobbyDetailsEOS::Ge
 	for (uint32_t MemberIndex = 0; MemberIndex < MemberCount; ++MemberIndex)
 	{
 		EOS_LobbyDetails_GetMemberByIndexOptions GetMemberByIndexOptions = {};
-		GetMemberByIndexOptions.ApiVersion = EOS_LOBBYDETAILS_GETMEMBERBYINDEX_API_LATEST;
+		GetMemberByIndexOptions.ApiVersion = 1;
+		UE_EOS_CHECK_API_MISMATCH(EOS_LOBBYDETAILS_GETMEMBERBYINDEX_API_LATEST, 1);
 		GetMemberByIndexOptions.MemberIndex = MemberIndex;
-		static_assert(EOS_LOBBYDETAILS_GETMEMBERBYINDEX_API_LATEST == 1, "EOS_LobbyDetails_GetMemberByIndexOptions updated, check new fields");
 
 		MemberProductUserIds->Emplace(EOS_LobbyDetails_GetMemberByIndex(LobbyDetailsHandle, &GetMemberByIndexOptions));
 	}
@@ -349,8 +349,8 @@ TFuture<TDefaultErrorResultInternal<FLobbyServiceSnapshot>> FLobbyDetailsEOS::Ge
 		// Resolve member info.
 		{
 			EOS_LobbyDetails_GetLobbyOwnerOptions GetLobbyOwnerOptions = {};
-			GetLobbyOwnerOptions.ApiVersion = EOS_LOBBYDETAILS_GETLOBBYOWNER_API_LATEST;
-			static_assert(EOS_LOBBYDETAILS_GETLOBBYOWNER_API_LATEST == 1, "EOS_LobbyDetails_GetLobbyOwnerOptions updated, check new fields");
+			GetLobbyOwnerOptions.ApiVersion = 1;
+			UE_EOS_CHECK_API_MISMATCH(EOS_LOBBYDETAILS_GETLOBBYOWNER_API_LATEST, 1);
 
 			const EOS_ProductUserId LobbyOwner = EOS_LobbyDetails_GetLobbyOwner(StrongThis->LobbyDetailsHandle, &GetLobbyOwnerOptions);
 
@@ -371,16 +371,16 @@ TFuture<TDefaultErrorResultInternal<FLobbyServiceSnapshot>> FLobbyDetailsEOS::Ge
 		// Resolve lobby attributes
 		{
 			EOS_LobbyDetails_GetAttributeCountOptions GetAttributeCountOptions = {};
-			GetAttributeCountOptions.ApiVersion = EOS_LOBBYDETAILS_GETATTRIBUTECOUNT_API_LATEST;
-			static_assert(EOS_LOBBYDETAILS_GETATTRIBUTECOUNT_API_LATEST == 1, "EOS_LobbyDetails_GetAttributeCountOptions updated, check new fields");
+			GetAttributeCountOptions.ApiVersion = 1;
+			UE_EOS_CHECK_API_MISMATCH(EOS_LOBBYDETAILS_GETATTRIBUTECOUNT_API_LATEST, 1);
 
 			const uint32_t AttributeCount = EOS_LobbyDetails_GetAttributeCount(StrongThis->LobbyDetailsHandle, &GetAttributeCountOptions);
 			for (uint32_t AttributeIndex = 0; AttributeIndex < AttributeCount; ++AttributeIndex)
 			{
 				EOS_LobbyDetails_CopyAttributeByIndexOptions CopyAttributeByIndexOptions = {};
-				CopyAttributeByIndexOptions.ApiVersion = EOS_LOBBYDETAILS_COPYATTRIBUTEBYINDEX_API_LATEST;
+				CopyAttributeByIndexOptions.ApiVersion = 1;
+				UE_EOS_CHECK_API_MISMATCH(EOS_LOBBYDETAILS_COPYATTRIBUTEBYINDEX_API_LATEST, 1);
 				CopyAttributeByIndexOptions.AttrIndex = AttributeIndex;
-				static_assert(EOS_LOBBYDETAILS_COPYATTRIBUTEBYINDEX_API_LATEST == 1, "EOS_LobbyDetails_CopyAttributeByIndexOptions updated, check new fields");
 
 				EOS_Lobby_Attribute* LobbyAttribute = nullptr;
 				ON_SCOPE_EXIT
@@ -423,18 +423,18 @@ TDefaultErrorResultInternal<FLobbyMemberServiceSnapshot> FLobbyDetailsEOS::GetLo
 	// Fetch attributes.
 	{
 		EOS_LobbyDetails_GetMemberAttributeCountOptions GetMemberAttributeCountOptions = {};
-		GetMemberAttributeCountOptions.ApiVersion = EOS_LOBBYDETAILS_GETMEMBERATTRIBUTECOUNT_API_LATEST;
+		GetMemberAttributeCountOptions.ApiVersion = 1;
+		UE_EOS_CHECK_API_MISMATCH(EOS_LOBBYDETAILS_GETMEMBERATTRIBUTECOUNT_API_LATEST, 1);
 		GetMemberAttributeCountOptions.TargetUserId = MemberProductUserId;
-		static_assert(EOS_LOBBYDETAILS_GETMEMBERATTRIBUTECOUNT_API_LATEST == 1, "EOS_LobbyDetails_GetMemberAttributeCountOptions updated, check new fields");
 
 		const uint32_t MemberAttributeCount = EOS_LobbyDetails_GetMemberAttributeCount(LobbyDetailsHandle, &GetMemberAttributeCountOptions);
 		for (uint32_t MemberAttributeIndex = 0; MemberAttributeIndex < MemberAttributeCount; ++MemberAttributeIndex)
 		{
 			EOS_LobbyDetails_CopyMemberAttributeByIndexOptions CopyMemberAttributeByIndexOptions = {};
-			CopyMemberAttributeByIndexOptions.ApiVersion = EOS_LOBBYDETAILS_COPYMEMBERATTRIBUTEBYINDEX_API_LATEST;
+			CopyMemberAttributeByIndexOptions.ApiVersion = 1;
+			UE_EOS_CHECK_API_MISMATCH(EOS_LOBBYDETAILS_COPYMEMBERATTRIBUTEBYINDEX_API_LATEST, 1);
 			CopyMemberAttributeByIndexOptions.TargetUserId = MemberProductUserId;
 			CopyMemberAttributeByIndexOptions.AttrIndex = MemberAttributeIndex;
-			static_assert(EOS_LOBBYDETAILS_COPYMEMBERATTRIBUTEBYINDEX_API_LATEST == 1, "EOS_LobbyDetails_CopyMemberAttributeByIndexOptions updated, check new fields");
 
 			EOS_Lobby_Attribute* LobbyAttribute = nullptr;
 			ON_SCOPE_EXIT
@@ -471,10 +471,10 @@ TFuture<EOS_EResult> FLobbyDetailsEOS::ApplyLobbyDataUpdateFromLocalChanges(FAcc
 
 	// Create lobby modification handle.
 	EOS_Lobby_UpdateLobbyModificationOptions ModificationOptions = {};
-	ModificationOptions.ApiVersion = EOS_LOBBY_UPDATELOBBYMODIFICATION_API_LATEST;
+	ModificationOptions.ApiVersion = 1;
+	UE_EOS_CHECK_API_MISMATCH(EOS_LOBBY_UPDATELOBBYMODIFICATION_API_LATEST, 1);
 	ModificationOptions.LocalUserId = GetProductUserIdChecked(LocalAccountId);
 	ModificationOptions.LobbyId = GetInfo()->GetLobbyIdEOS();
-	static_assert(EOS_LOBBY_UPDATELOBBYMODIFICATION_API_LATEST == 1, "EOS_Lobby_UpdateLobbyModificationOptions updated, check new fields");
 
 	EOS_EResult EOSResult = EOS_Lobby_UpdateLobbyModification(Prerequisites->LobbyInterfaceHandle, &ModificationOptions, &LobbyModificationHandle);
 	if (EOSResult != EOS_EResult::EOS_Success)
@@ -488,9 +488,9 @@ TFuture<EOS_EResult> FLobbyDetailsEOS::ApplyLobbyDataUpdateFromLocalChanges(FAcc
 	if (ServiceChanges.JoinPolicy)
 	{
 		EOS_LobbyModification_SetPermissionLevelOptions SetPermissionOptions = {};
-		SetPermissionOptions.ApiVersion = EOS_LOBBYMODIFICATION_SETPERMISSIONLEVEL_API_LATEST;
+		SetPermissionOptions.ApiVersion = 1;
+		UE_EOS_CHECK_API_MISMATCH(EOS_LOBBYMODIFICATION_SETPERMISSIONLEVEL_API_LATEST, 1);
 		SetPermissionOptions.PermissionLevel = TranslateJoinPolicy(*ServiceChanges.JoinPolicy);
-		static_assert(EOS_LOBBYMODIFICATION_SETPERMISSIONLEVEL_API_LATEST == 1, "EOS_LobbyModification_SetPermissionLevelOptions updated, check new fields");
 
 		EOSResult = EOS_LobbyModification_SetPermissionLevel(LobbyModificationHandle, &SetPermissionOptions);
 		if (EOSResult != EOS_EResult::EOS_Success)
@@ -507,10 +507,10 @@ TFuture<EOS_EResult> FLobbyDetailsEOS::ApplyLobbyDataUpdateFromLocalChanges(FAcc
 		const FLobbyAttributeTranslator<ELobbyTranslationType::ToService> AttributeTranslator(UpdatedAttribute.Value);
 
 		EOS_LobbyModification_AddAttributeOptions AddAttributeOptions = {};
-		AddAttributeOptions.ApiVersion = EOS_LOBBYMODIFICATION_ADDATTRIBUTE_API_LATEST;
+		AddAttributeOptions.ApiVersion = 1;
+		UE_EOS_CHECK_API_MISMATCH(EOS_LOBBYMODIFICATION_ADDATTRIBUTE_API_LATEST, 1);
 		AddAttributeOptions.Attribute = &AttributeTranslator.GetAttributeData();
 		AddAttributeOptions.Visibility = AttributeTranslator.GetAttributeVisibility();
-		static_assert(EOS_LOBBYMODIFICATION_ADDATTRIBUTE_API_LATEST == 1, "EOS_LobbyModification_AddAttributeOptions updated, check new fields");
 
 		EOSResult = EOS_LobbyModification_AddAttribute(LobbyModificationHandle, &AddAttributeOptions);
 		if (EOSResult != EOS_EResult::EOS_Success)
@@ -527,9 +527,9 @@ TFuture<EOS_EResult> FLobbyDetailsEOS::ApplyLobbyDataUpdateFromLocalChanges(FAcc
 		const FTCHARToUTF8 KeyConverter(*RemovedAttribute.ToString());
 
 		EOS_LobbyModification_RemoveAttributeOptions RemoveAttributeOptions = {};
-		RemoveAttributeOptions.ApiVersion = EOS_LOBBYMODIFICATION_REMOVEATTRIBUTE_API_LATEST;
+		RemoveAttributeOptions.ApiVersion = 1;
+		UE_EOS_CHECK_API_MISMATCH(EOS_LOBBYMODIFICATION_REMOVEATTRIBUTE_API_LATEST, 1);
 		RemoveAttributeOptions.Key = KeyConverter.Get();
-		static_assert(EOS_LOBBYMODIFICATION_REMOVEATTRIBUTE_API_LATEST == 1, "EOS_LobbyModification_RemoveAttributeOptions updated, check new fields");
 		
 		EOSResult = EOS_LobbyModification_RemoveAttribute(LobbyModificationHandle, &RemoveAttributeOptions);
 		if (EOSResult != EOS_EResult::EOS_Success)
@@ -546,10 +546,10 @@ TFuture<EOS_EResult> FLobbyDetailsEOS::ApplyLobbyDataUpdateFromLocalChanges(FAcc
 		const FLobbyAttributeTranslator<ELobbyTranslationType::ToService> AttributeTranslator(UpdatedAttribute.Value);
 
 		EOS_LobbyModification_AddMemberAttributeOptions AddMemberAttributeOptions = {};
-		AddMemberAttributeOptions.ApiVersion = EOS_LOBBYMODIFICATION_ADDMEMBERATTRIBUTE_API_LATEST;
+		AddMemberAttributeOptions.ApiVersion = 1;
+		UE_EOS_CHECK_API_MISMATCH(EOS_LOBBYMODIFICATION_ADDMEMBERATTRIBUTE_API_LATEST, 1);
 		AddMemberAttributeOptions.Attribute = &AttributeTranslator.GetAttributeData();
 		AddMemberAttributeOptions.Visibility = AttributeTranslator.GetAttributeVisibility();
-		static_assert(EOS_LOBBYMODIFICATION_ADDMEMBERATTRIBUTE_API_LATEST == 1, "EOS_LobbyModification_AddMemberAttributeOptions updated, check new fields");
 
 		EOSResult = EOS_LobbyModification_AddMemberAttribute(LobbyModificationHandle, &AddMemberAttributeOptions);
 		if (EOSResult != EOS_EResult::EOS_Success)
@@ -566,9 +566,9 @@ TFuture<EOS_EResult> FLobbyDetailsEOS::ApplyLobbyDataUpdateFromLocalChanges(FAcc
 		const FTCHARToUTF8 KeyConverter(*RemovedAttribute.ToString());
 
 		EOS_LobbyModification_RemoveMemberAttributeOptions RemoveMemberAttributeOptions = {};
-		RemoveMemberAttributeOptions.ApiVersion = EOS_LOBBYMODIFICATION_REMOVEMEMBERATTRIBUTE_API_LATEST;
+		RemoveMemberAttributeOptions.ApiVersion = 1;
+		UE_EOS_CHECK_API_MISMATCH(EOS_LOBBYMODIFICATION_REMOVEMEMBERATTRIBUTE_API_LATEST, 1);
 		RemoveMemberAttributeOptions.Key = KeyConverter.Get();
-		static_assert(EOS_LOBBYMODIFICATION_REMOVEMEMBERATTRIBUTE_API_LATEST == 1, "EOS_LobbyModification_RemoveMemberAttributeOptions updated, check new fields");
 
 		EOSResult = EOS_LobbyModification_RemoveMemberAttribute(LobbyModificationHandle, &RemoveMemberAttributeOptions);
 		if (EOSResult != EOS_EResult::EOS_Success)
@@ -584,9 +584,9 @@ TFuture<EOS_EResult> FLobbyDetailsEOS::ApplyLobbyDataUpdateFromLocalChanges(FAcc
 
 	// Apply lobby updates.
 	EOS_Lobby_UpdateLobbyOptions UpdateLobbyOptions = {};
-	UpdateLobbyOptions.ApiVersion = EOS_LOBBY_UPDATELOBBY_API_LATEST;
+	UpdateLobbyOptions.ApiVersion = 1;
+	UE_EOS_CHECK_API_MISMATCH(EOS_LOBBY_UPDATELOBBY_API_LATEST, 1);
 	UpdateLobbyOptions.LobbyModificationHandle = LobbyModificationHandle;
-	static_assert(EOS_LOBBY_UPDATELOBBY_API_LATEST == 1, "EOS_Lobby_UpdateLobbyOptions updated, check new fields");
 
 	EOS_Async(EOS_Lobby_UpdateLobby, Prerequisites->LobbyInterfaceHandle, UpdateLobbyOptions,
 	[LocalAccountId, LobbyInfo = GetInfo(), Promise = MoveTemp(Promise)](const EOS_Lobby_UpdateLobbyCallbackInfo* Data) mutable
@@ -624,8 +624,8 @@ FLobbyDetailsEOS::FLobbyDetailsEOS(
 TDefaultErrorResultInternal<TSharedRef<FLobbyDetailsInfoEOS>> FLobbyDetailsInfoEOS::Create(EOS_HLobbyDetails LobbyDetailsHandle)
 {
 	EOS_LobbyDetails_CopyInfoOptions CopyInfoOptions = {};
-	CopyInfoOptions.ApiVersion = EOS_LOBBYDETAILS_COPYINFO_API_LATEST;
-	static_assert(EOS_LOBBYDETAILS_COPYINFO_API_LATEST == 1, "EOS_LobbyDetails_CopyInfoOptions updated, check new fields");
+	CopyInfoOptions.ApiVersion = 1;
+	UE_EOS_CHECK_API_MISMATCH(EOS_LOBBYDETAILS_COPYINFO_API_LATEST, 1);
 
 	EOS_LobbyDetails_Info* LobbyDetailsInfo = nullptr;
 	EOS_EResult EOSResult = EOS_LobbyDetails_CopyInfo(LobbyDetailsHandle, &CopyInfoOptions, &LobbyDetailsInfo);
@@ -946,9 +946,9 @@ TFuture<TDefaultErrorResultInternal<TSharedRef<FLobbySearchEOS>>> FLobbySearchEO
 	TSharedRef<FSearchHandle> SearchHandle = MakeShared<FSearchHandle>();
 
 	EOS_Lobby_CreateLobbySearchOptions CreateLobbySearchOptions = {};
-	CreateLobbySearchOptions.ApiVersion = EOS_LOBBY_CREATELOBBYSEARCH_API_LATEST;
+	CreateLobbySearchOptions.ApiVersion = 1;
+	UE_EOS_CHECK_API_MISMATCH(EOS_LOBBY_CREATELOBBYSEARCH_API_LATEST, 1);
 	CreateLobbySearchOptions.MaxResults = Params.MaxResults;
-	static_assert(EOS_LOBBY_CREATELOBBYSEARCH_API_LATEST == 1, "EOS_Lobby_CreateLobbySearchOptions updated, check new fields");
 
 	EOS_EResult EOSResult = EOS_Lobby_CreateLobbySearch(Prerequisites->LobbyInterfaceHandle, &CreateLobbySearchOptions, &SearchHandle->Get());
 	if (EOSResult != EOS_EResult::EOS_Success)
@@ -969,9 +969,9 @@ TFuture<TDefaultErrorResultInternal<TSharedRef<FLobbySearchEOS>>> FLobbySearchEO
 		}
 
 		EOS_LobbySearch_SetLobbyIdOptions SetLobbyIdOptions = {};
-		SetLobbyIdOptions.ApiVersion = EOS_LOBBYSEARCH_SETLOBBYID_API_LATEST;
+		SetLobbyIdOptions.ApiVersion = 1;
+		UE_EOS_CHECK_API_MISMATCH(EOS_LOBBYSEARCH_SETLOBBYID_API_LATEST, 1);
 		SetLobbyIdOptions.LobbyId = LobbyData->GetLobbyIdEOS();
-		static_assert(EOS_LOBBYSEARCH_SETLOBBYID_API_LATEST == 1, "EOS_LobbySearch_SetLobbyIdOptions updated, check new fields");
 
 		EOSResult = EOS_LobbySearch_SetLobbyId(SearchHandle->Get(), &SetLobbyIdOptions);
 		if (EOSResult != EOS_EResult::EOS_Success)
@@ -984,9 +984,9 @@ TFuture<TDefaultErrorResultInternal<TSharedRef<FLobbySearchEOS>>> FLobbySearchEO
 	else if (Params.TargetUser) // Search for specific user.
 	{
 		EOS_LobbySearch_SetTargetUserIdOptions SetTargetUserIdOptions = {};
-		SetTargetUserIdOptions.ApiVersion = EOS_LOBBYSEARCH_SETTARGETUSERID_API_LATEST;
+		SetTargetUserIdOptions.ApiVersion = 1;
+		UE_EOS_CHECK_API_MISMATCH(EOS_LOBBYSEARCH_SETTARGETUSERID_API_LATEST, 1);
 		SetTargetUserIdOptions.TargetUserId = GetProductUserIdChecked(*Params.TargetUser);
-		static_assert(EOS_LOBBYSEARCH_SETTARGETUSERID_API_LATEST == 1, "EOS_LobbySearch_SetTargetUserIdOptions updated, check new fields");
 
 		EOSResult = EOS_LobbySearch_SetTargetUserId(SearchHandle->Get(), &SetTargetUserIdOptions);
 		if (EOSResult != EOS_EResult::EOS_Success)
@@ -1003,17 +1003,17 @@ TFuture<TDefaultErrorResultInternal<TSharedRef<FLobbySearchEOS>>> FLobbySearchEO
 			const FLobbyBucketIdTranslator<ELobbyTranslationType::ToService> BucketTranslator(Prerequisites->BucketId);
 
 			EOS_Lobby_AttributeData AttributeData;
-			AttributeData.ApiVersion = EOS_LOBBY_ATTRIBUTEDATA_API_LATEST;
+			AttributeData.ApiVersion = 1;
 			AttributeData.Key = EOS_LOBBY_SEARCH_BUCKET_ID;
 			AttributeData.ValueType = EOS_ELobbyAttributeType::EOS_AT_STRING;
 			AttributeData.Value.AsUtf8 = BucketTranslator.GetBucketIdEOS();
-			static_assert(EOS_LOBBY_ATTRIBUTEDATA_API_LATEST == 1, "EOS_Lobby_AttributeData updated, check new fields");
+			UE_EOS_CHECK_API_MISMATCH(EOS_LOBBY_ATTRIBUTEDATA_API_LATEST, 1);
 
 			EOS_LobbySearch_SetParameterOptions SetParameterOptions = {};
-			SetParameterOptions.ApiVersion = EOS_LOBBYSEARCH_SETPARAMETER_API_LATEST;
+			SetParameterOptions.ApiVersion = 1;
 			SetParameterOptions.Parameter = &AttributeData;
 			SetParameterOptions.ComparisonOp = EOS_EComparisonOp::EOS_CO_EQUAL;
-			static_assert(EOS_LOBBYSEARCH_SETPARAMETER_API_LATEST == 1, "EOS_LobbySearch_SetParameterOptions updated, check new fields");
+			UE_EOS_CHECK_API_MISMATCH(EOS_LOBBYSEARCH_SETPARAMETER_API_LATEST, 1);
 
 			EOSResult = EOS_LobbySearch_SetParameter(SearchHandle->Get(), &SetParameterOptions);
 			if (EOSResult != EOS_EResult::EOS_Success)
@@ -1041,10 +1041,10 @@ TFuture<TDefaultErrorResultInternal<TSharedRef<FLobbySearchEOS>>> FLobbySearchEO
 				{ SchemaServiceAttributeId, SchemaServiceAttributeFlags, Filter.ComparisonValue });
 
 			EOS_LobbySearch_SetParameterOptions SetParameterOptions = {};
-			SetParameterOptions.ApiVersion = EOS_LOBBYSEARCH_SETPARAMETER_API_LATEST;
+			SetParameterOptions.ApiVersion = 1;
 			SetParameterOptions.Parameter = &AttributeTranslator.GetAttributeData();
 			SetParameterOptions.ComparisonOp = Private::TranslateSearchComparison(Filter.ComparisonOp);
-			static_assert(EOS_LOBBYSEARCH_SETPARAMETER_API_LATEST == 1, "EOS_LobbySearch_SetParameterOptions updated, check new fields");
+			UE_EOS_CHECK_API_MISMATCH(EOS_LOBBYSEARCH_SETPARAMETER_API_LATEST, 1);
 
 			EOSResult = EOS_LobbySearch_SetParameter(SearchHandle->Get(), &SetParameterOptions);
 			if (EOSResult != EOS_EResult::EOS_Success)
@@ -1060,9 +1060,9 @@ TFuture<TDefaultErrorResultInternal<TSharedRef<FLobbySearchEOS>>> FLobbySearchEO
 	TFuture<TDefaultErrorResultInternal<TSharedRef<FLobbySearchEOS>>> Future = Promise.GetFuture();
 
 	EOS_LobbySearch_FindOptions FindOptions = {};
-	FindOptions.ApiVersion = EOS_LOBBYSEARCH_FIND_API_LATEST;
+	FindOptions.ApiVersion = 1;
 	FindOptions.LocalUserId = GetProductUserIdChecked(Params.LocalAccountId);
-	static_assert(EOS_LOBBYSEARCH_FIND_API_LATEST == 1, "EOS_LobbySearch_FindOptions updated, check new fields");
+	UE_EOS_CHECK_API_MISMATCH(EOS_LOBBYSEARCH_FIND_API_LATEST, 1);
 
 	EOS_Async(EOS_LobbySearch_Find, SearchHandle->Get(), FindOptions,
 	[
@@ -1084,8 +1084,8 @@ TFuture<TDefaultErrorResultInternal<TSharedRef<FLobbySearchEOS>>> FLobbySearchEO
 		TArray<TFuture<TDefaultErrorResultInternal<TSharedRef<FLobbyDataEOS>>>> ResolvedLobbyDetails;
 
 		EOS_LobbySearch_GetSearchResultCountOptions GetSearchResultCountOptions = {};
-		GetSearchResultCountOptions.ApiVersion = EOS_LOBBYSEARCH_GETSEARCHRESULTCOUNT_API_LATEST;
-		static_assert(EOS_LOBBYSEARCH_GETSEARCHRESULTCOUNT_API_LATEST == 1, "EOS_LobbySearch_GetSearchResultCountOptions updated, check new fields");
+		GetSearchResultCountOptions.ApiVersion = 1;
+		UE_EOS_CHECK_API_MISMATCH(EOS_LOBBYSEARCH_GETSEARCHRESULTCOUNT_API_LATEST, 1);
 		
 		const uint32_t NumSearchResults = EOS_LobbySearch_GetSearchResultCount(SearchHandle->Get(), &GetSearchResultCountOptions);
 
diff --git a/Engine/Plugins/Online/OnlineServicesEOSGS/Source/Private/Online/OnlineServicesEOSGSPlatformFactory.cpp b/Engine/Plugins/Online/OnlineServicesEOSGS/Source/Private/Online/OnlineServicesEOSGSPlatformFactory.cpp
--- a/Engine/Plugins/Online/OnlineServicesEOSGS/Source/Private/Online/OnlineServicesEOSGSPlatformFactory.cpp
+++ b/Engine/Plugins/Online/OnlineServicesEOSGS/Source/Private/Online/OnlineServicesEOSGSPlatformFactory.cpp
@@ -98,7 +98,8 @@ IEOSPlatformHandlePtr FOnlineServicesEOSGSPlatformFactory::CreatePlatform()
 	const FTCHARToUTF8 CacheDirectory(*(SDKManager->GetCacheDirBase() / TEXT("OnlineServicesEOS")));
 
 	EOS_Platform_Options PlatformOptions = {};
-	PlatformOptions.ApiVersion = EOS_PLATFORM_OPTIONS_API_LATEST;
+	PlatformOptions.ApiVersion = 12;
+	UE_EOS_CHECK_API_MISMATCH(EOS_PLATFORM_OPTIONS_API_LATEST, 12);
 	PlatformOptions.Reserved = nullptr;
 	PlatformOptions.bIsServer = EOS_FALSE;
 	PlatformOptions.OverrideCountryCode = nullptr;
diff --git a/Engine/Plugins/Online/OnlineServicesEOSGS/Source/Private/Online/SessionsEOSGS.cpp b/Engine/Plugins/Online/OnlineServicesEOSGS/Source/Private/Online/SessionsEOSGS.cpp
--- a/Engine/Plugins/Online/OnlineServicesEOSGS/Source/Private/Online/SessionsEOSGS.cpp
+++ b/Engine/Plugins/Online/OnlineServicesEOSGS/Source/Private/Online/SessionsEOSGS.cpp
@@ -46,8 +46,8 @@ FSessionEOSGS::FSessionEOSGS(const TSharedPtr<FSessionDetailsHandleEOSGS>& InSes
 	: SessionDetailsHandle(InSessionDetailsHandle)
 {
 	EOS_SessionDetails_CopyInfoOptions CopyInfoOptions = { };
-	CopyInfoOptions.ApiVersion = EOS_SESSIONDETAILS_COPYINFO_API_LATEST;
-	static_assert(EOS_SESSIONDETAILS_COPYINFO_API_LATEST == 1, "EOS_SessionDetails_CopyInfoOptions updated, check new fields");
+	CopyInfoOptions.ApiVersion = 1;
+	UE_EOS_CHECK_API_MISMATCH(EOS_SESSIONDETAILS_COPYINFO_API_LATEST, 1);
 
 	EOS_SessionDetails_Info* SessionDetailsInfo = nullptr;
 	EOS_EResult CopyInfoResult = EOS_SessionDetails_CopyInfo(SessionDetailsHandle->SessionDetailsHandle, &CopyInfoOptions, &SessionDetailsInfo);
@@ -65,8 +65,8 @@ FSessionEOSGS::FSessionEOSGS(const TSharedPtr<FSessionDetailsHandleEOSGS>& InSes
 
 		// We retrieve all the session attributes
 		EOS_SessionDetails_GetSessionAttributeCountOptions GetAttributeCountOptions = {};
-		GetAttributeCountOptions.ApiVersion = EOS_SESSIONDETAILS_GETSESSIONATTRIBUTECOUNT_API_LATEST;
-		static_assert(EOS_SESSIONDETAILS_GETSESSIONATTRIBUTECOUNT_API_LATEST == 1, "EOS_SessionDetails_GetSessionAttributeCountOptions updated, check new fields");
+		GetAttributeCountOptions.ApiVersion = 1;
+		UE_EOS_CHECK_API_MISMATCH(EOS_SESSIONDETAILS_GETSESSIONATTRIBUTECOUNT_API_LATEST, 1);
 
 		uint32_t AttributeCount = EOS_SessionDetails_GetSessionAttributeCount(SessionDetailsHandle->SessionDetailsHandle, &GetAttributeCountOptions);
 		for (uint32_t Index = 0; Index < AttributeCount; ++Index)
@@ -74,8 +74,8 @@ FSessionEOSGS::FSessionEOSGS(const TSharedPtr<FSessionDetailsHandleEOSGS>& InSes
 			// We parse a single attribute
 
 			EOS_SessionDetails_CopySessionAttributeByIndexOptions CopyAttributeByIndexOptions = {};
-			CopyAttributeByIndexOptions.ApiVersion = EOS_SESSIONDETAILS_COPYSESSIONATTRIBUTEBYINDEX_API_LATEST;
-			static_assert(EOS_SESSIONDETAILS_COPYSESSIONATTRIBUTEBYINDEX_API_LATEST == 1, "EOS_SessionDetails_CopySessionAttributeByIndexOptions updated, check new fields");
+			CopyAttributeByIndexOptions.ApiVersion = 1;
+			UE_EOS_CHECK_API_MISMATCH(EOS_SESSIONDETAILS_COPYSESSIONATTRIBUTEBYINDEX_API_LATEST, 1);
 
 			CopyAttributeByIndexOptions.AttrIndex = Index;
 
@@ -162,34 +162,31 @@ void FSessionsEOSGS::RegisterEventHandlers()
 	OnSessionInviteReceivedEventRegistration = EOS_RegisterComponentEventHandler(
 		this,
 		SessionsHandle,
-		EOS_SESSIONS_ADDNOTIFYSESSIONINVITERECEIVED_API_LATEST,
+		1,
 		&EOS_Sessions_AddNotifySessionInviteReceived,
 		&EOS_Sessions_RemoveNotifySessionInviteReceived,
 		&FSessionsEOSGS::HandleSessionInviteReceived);
-
-	static_assert(EOS_SESSIONS_ADDNOTIFYSESSIONINVITERECEIVED_API_LATEST == 1, "EOS_Sessions_AddNotifySessionInviteReceivedOptions updated, check new fields");
+	UE_EOS_CHECK_API_MISMATCH(EOS_SESSIONS_ADDNOTIFYSESSIONINVITERECEIVED_API_LATEST, 1);
 
 	// Register for session invites accepted events
 	OnSessionInviteAcceptedEventRegistration = EOS_RegisterComponentEventHandler(
 		this,
 		SessionsHandle,
-		EOS_SESSIONS_ADDNOTIFYSESSIONINVITEACCEPTED_API_LATEST,
+		1,
 		&EOS_Sessions_AddNotifySessionInviteAccepted,
 		&EOS_Sessions_RemoveNotifySessionInviteAccepted,
 		&FSessionsEOSGS::HandleSessionInviteAccepted);
-
-	static_assert(EOS_SESSIONS_ADDNOTIFYSESSIONINVITEACCEPTED_API_LATEST == 1, "EOS_Sessions_AddNotifySessionInviteAcceptedOptions updated, check new fields");
+	UE_EOS_CHECK_API_MISMATCH(EOS_SESSIONS_ADDNOTIFYSESSIONINVITEACCEPTED_API_LATEST, 1);
 
 	// Register for join session accepted events
 	OnJoinSessionAcceptedEventRegistration = EOS_RegisterComponentEventHandler(
 		this,
 		SessionsHandle,
-		EOS_SESSIONS_ADDNOTIFYJOINSESSIONACCEPTED_API_LATEST,
+		1,
 		&EOS_Sessions_AddNotifyJoinSessionAccepted,
 		&EOS_Sessions_RemoveNotifyJoinSessionAccepted,
 		&FSessionsEOSGS::HandleJoinSessionAccepted);
-
-	static_assert(EOS_SESSIONS_ADDNOTIFYJOINSESSIONACCEPTED_API_LATEST == 1, "EOS_Sessions_AddNotifyJoinSessionAcceptedOptions updated, check new fields");
+	UE_EOS_CHECK_API_MISMATCH(EOS_SESSIONS_ADDNOTIFYJOINSESSIONACCEPTED_API_LATEST, 1);
 }
 
 void FSessionsEOSGS::UnregisterEventHandlers()
@@ -327,8 +324,8 @@ TFuture<TOnlineResult<FCreateSession>> FSessionsEOSGS::CreateSessionImpl(const F
 	TFuture<TOnlineResult<FCreateSession>> Future = Promise.GetFuture();
 
 	EOS_Sessions_CreateSessionModificationOptions CreateSessionModificationOptions = {};
-	CreateSessionModificationOptions.ApiVersion = EOS_SESSIONS_CREATESESSIONMODIFICATION_API_LATEST;
-	static_assert(EOS_SESSIONS_CREATESESSIONMODIFICATION_API_LATEST == 4, "EOS_Sessions_CreateSessionModificationOptions updated, check new fields");
+	CreateSessionModificationOptions.ApiVersion = 4;
+	UE_EOS_CHECK_API_MISMATCH(EOS_SESSIONS_CREATESESSIONMODIFICATION_API_LATEST, 4);
 
 	CreateSessionModificationOptions.bPresenceEnabled = Params.bPresenceEnabled;
 	CreateSessionModificationOptions.bSanctionsEnabled = !Params.bAllowSanctionedPlayers;
@@ -434,8 +431,8 @@ TOptional<FOnlineError> FSessionsEOSGS::CheckState(const FCreateSession::Params&
 void FSessionsEOSGS::SetHostAddress(EOS_HSessionModification& SessionModHandle, FString HostAddress)
 {
 	EOS_SessionModification_SetHostAddressOptions Options = {};
-	Options.ApiVersion = EOS_SESSIONMODIFICATION_SETHOSTADDRESS_API_LATEST;
-	static_assert(EOS_SESSIONMODIFICATION_SETHOSTADDRESS_API_LATEST == 1, "EOS_SessionModification_SetHostAddressOptions updated, check new fields");
+	Options.ApiVersion = 1;
+	UE_EOS_CHECK_API_MISMATCH(EOS_SESSIONMODIFICATION_SETHOSTADDRESS_API_LATEST, 1);
 
 	const FTCHARToUTF8 HostAddressUtf8(*HostAddress);
 	Options.HostAddress = HostAddressUtf8.Get();
@@ -450,8 +447,8 @@ void FSessionsEOSGS::SetHostAddress(EOS_HSessionModification& SessionModHandle,
 void FSessionsEOSGS::SetJoinInProgressAllowed(EOS_HSessionModification& SessionModHandle, bool bIsJoinInProgressAllowed)
 {
 	EOS_SessionModification_SetJoinInProgressAllowedOptions Options = {};
-	Options.ApiVersion = EOS_SESSIONMODIFICATION_SETJOININPROGRESSALLOWED_API_LATEST;
-	static_assert(EOS_SESSIONMODIFICATION_SETJOININPROGRESSALLOWED_API_LATEST == 1, "EOS_SessionModification_SetJoinInProgressAllowedOptions updated, check new fields");
+	Options.ApiVersion = 1;
+	UE_EOS_CHECK_API_MISMATCH(EOS_SESSIONMODIFICATION_SETJOININPROGRESSALLOWED_API_LATEST, 1);
 
 	Options.bAllowJoinInProgress = bIsJoinInProgressAllowed;
 
@@ -465,8 +462,8 @@ void FSessionsEOSGS::SetJoinInProgressAllowed(EOS_HSessionModification& SessionM
 void FSessionsEOSGS::SetInvitesAllowed(EOS_HSessionModification& SessionModHandle, bool bAreInvitesAllowed)
 {
 	EOS_SessionModification_SetInvitesAllowedOptions Options = {};
-	Options.ApiVersion = EOS_SESSIONMODIFICATION_SETINVITESALLOWED_API_LATEST;
-	static_assert(EOS_SESSIONMODIFICATION_SETINVITESALLOWED_API_LATEST == 1, "EOS_SessionModification_SetInvitesAllowedOptions updated, check new fields");
+	Options.ApiVersion = 1;
+	UE_EOS_CHECK_API_MISMATCH(EOS_SESSIONMODIFICATION_SETINVITESALLOWED_API_LATEST, 1);
 
 	Options.bInvitesAllowed = bAreInvitesAllowed;
 
@@ -480,8 +477,8 @@ void FSessionsEOSGS::SetInvitesAllowed(EOS_HSessionModification& SessionModHandl
 void FSessionsEOSGS::SetPermissionLevel(EOS_HSessionModification& SessionModificationHandle, const ESessionJoinPolicy& NewJoinPolicy)
 {
 	EOS_SessionModification_SetPermissionLevelOptions Options = { };
-	Options.ApiVersion = EOS_SESSIONMODIFICATION_SETPERMISSIONLEVEL_API_LATEST;
-	static_assert(EOS_SESSIONMODIFICATION_SETPERMISSIONLEVEL_API_LATEST == 1, "EOS_SessionModification_SetPermissionLevelOptions updated, check new fields");
+	Options.ApiVersion = 1;
+	UE_EOS_CHECK_API_MISMATCH(EOS_SESSIONMODIFICATION_SETPERMISSIONLEVEL_API_LATEST, 1);
 
 	Options.PermissionLevel = ToServiceType(NewJoinPolicy);
 
@@ -495,8 +492,8 @@ void FSessionsEOSGS::SetPermissionLevel(EOS_HSessionModification& SessionModific
 void FSessionsEOSGS::SetBucketId(EOS_HSessionModification& SessionModificationHandle, const FString& NewBucketId)
 {
 	EOS_SessionModification_SetBucketIdOptions Options = { };
-	Options.ApiVersion = EOS_SESSIONMODIFICATION_SETBUCKETID_API_LATEST;
-	static_assert(EOS_SESSIONMODIFICATION_SETBUCKETID_API_LATEST == 1, "EOS_SessionModification_SetBucketIdOptions updated, check new fields");
+	Options.ApiVersion = 1;
+	UE_EOS_CHECK_API_MISMATCH(EOS_SESSIONMODIFICATION_SETBUCKETID_API_LATEST, 1);
 
 	const FTCHARToUTF8 BucketIdUtf8(*NewBucketId);
 	Options.BucketId = BucketIdUtf8.Get();
@@ -511,8 +508,8 @@ void FSessionsEOSGS::SetBucketId(EOS_HSessionModification& SessionModificationHa
 void FSessionsEOSGS::SetMaxPlayers(EOS_HSessionModification& SessionModificationHandle, const uint32& NewMaxPlayers)
 {
 	EOS_SessionModification_SetMaxPlayersOptions Options = { };
-	Options.ApiVersion = EOS_SESSIONMODIFICATION_SETMAXPLAYERS_API_LATEST;
-	static_assert(EOS_SESSIONMODIFICATION_SETMAXPLAYERS_API_LATEST == 1, "EOS_SessionModification_SetMaxPlayersOptions updated, check new fields");
+	Options.ApiVersion = 1;
+	UE_EOS_CHECK_API_MISMATCH(EOS_SESSIONMODIFICATION_SETMAXPLAYERS_API_LATEST, 1);
 
 	Options.MaxPlayers = NewMaxPlayers;
 
@@ -526,8 +523,8 @@ void FSessionsEOSGS::SetMaxPlayers(EOS_HSessionModification& SessionModification
 void FSessionsEOSGS::AddAttribute(EOS_HSessionModification& SessionModificationHandle, const FSchemaAttributeId& Key, const FCustomSessionSetting& Value)
 {
 	EOS_SessionModification_AddAttributeOptions Options = { };
-	Options.ApiVersion = EOS_SESSIONMODIFICATION_ADDATTRIBUTE_API_LATEST;
-	static_assert(EOS_SESSIONMODIFICATION_ADDATTRIBUTE_API_LATEST == 1, "EOS_SessionModification_AddAttributeOptions updated, check new fields");
+	Options.ApiVersion = 1;
+	UE_EOS_CHECK_API_MISMATCH(EOS_SESSIONMODIFICATION_ADDATTRIBUTE_API_LATEST, 1);
 
 	Options.AdvertisementType = ToServiceType(Value.Visibility);
 
@@ -544,8 +541,8 @@ void FSessionsEOSGS::AddAttribute(EOS_HSessionModification& SessionModificationH
 void FSessionsEOSGS::RemoveAttribute(EOS_HSessionModification& SessionModificationHandle, const FSchemaAttributeId& Key)
 {
 	EOS_SessionModification_RemoveAttributeOptions Options = { };
-	Options.ApiVersion = EOS_SESSIONMODIFICATION_REMOVEATTRIBUTE_API_LATEST;
-	static_assert(EOS_SESSIONMODIFICATION_REMOVEATTRIBUTE_API_LATEST == 1, "EOS_SessionModification_RemoveAttributeOptions updated, check new fields");
+	Options.ApiVersion = 1;
+	UE_EOS_CHECK_API_MISMATCH(EOS_SESSIONMODIFICATION_REMOVEATTRIBUTE_API_LATEST, 1);
 
 	const FTCHARToUTF8 KeyUtf8(Key.ToString());
 	Options.Key = KeyUtf8.Get();
@@ -672,8 +669,8 @@ TFuture<TOnlineResult<FUpdateSessionSettings>> FSessionsEOSGS::UpdateSessionSett
 	TFuture<TOnlineResult<FUpdateSessionSettings>> Future = Promise.GetFuture();
 
 	EOS_Sessions_UpdateSessionModificationOptions UpdateSessionModificationOptions = {};
-	UpdateSessionModificationOptions.ApiVersion = EOS_SESSIONS_UPDATESESSIONMODIFICATION_API_LATEST;
-	static_assert(EOS_SESSIONS_UPDATESESSIONMODIFICATION_API_LATEST == 1, "EOS_Sessions_UpdateSessionModificationOptions updated, check new fields");
+	UpdateSessionModificationOptions.ApiVersion = 1;
+	UE_EOS_CHECK_API_MISMATCH(EOS_SESSIONS_UPDATESESSIONMODIFICATION_API_LATEST, 1);
 
 	const FTCHARToUTF8 SessionNameUtf8(*Params.SessionName.ToString());
 	UpdateSessionModificationOptions.SessionName = SessionNameUtf8.Get();
@@ -739,8 +736,8 @@ TFuture<TOnlineResult<FUpdateSessionSettings>> FSessionsEOSGS::UpdateSessionSett
 TFuture<TDefaultErrorResult<FUpdateSessionImplEOSGS>> FSessionsEOSGS::UpdateSessionImplEOSGS(FUpdateSessionImplEOSGS::Params&& Params)
 {
 	EOS_Sessions_UpdateSessionOptions UpdateSessionOptions = {};
-	UpdateSessionOptions.ApiVersion = EOS_SESSIONS_UPDATESESSION_API_LATEST;
-	static_assert(EOS_SESSIONS_UPDATESESSION_API_LATEST == 1, "EOS_Sessions_UpdateSessionOptions updated, check new fields");
+	UpdateSessionOptions.ApiVersion = 1;
+	UE_EOS_CHECK_API_MISMATCH(EOS_SESSIONS_UPDATESESSION_API_LATEST, 1);
 
 	UpdateSessionOptions.SessionModificationHandle = Params.SessionModificationHandle->ModificationHandle;
 
@@ -790,8 +787,8 @@ TFuture<TDefaultErrorResult<FUpdateSessionJoinabilityImpl>> FSessionsEOSGS::Upda
 
 	// We get the active session handle with the session name
 	EOS_Sessions_CopyActiveSessionHandleOptions CopyActiveSessionHandleOptions = {};
-	CopyActiveSessionHandleOptions.ApiVersion = EOS_SESSIONS_COPYACTIVESESSIONHANDLE_API_LATEST;
-	static_assert(EOS_SESSIONS_COPYACTIVESESSIONHANDLE_API_LATEST == 1, "EOS_Sessions_CopyActiveSessionHandleOptions updated, check new fields");
+	CopyActiveSessionHandleOptions.ApiVersion = 1;
+	UE_EOS_CHECK_API_MISMATCH(EOS_SESSIONS_COPYACTIVESESSIONHANDLE_API_LATEST, 1);
 
 	const FTCHARToUTF8 SessionNameUtf8(*Params.SessionName.ToString());
 	CopyActiveSessionHandleOptions.SessionName = SessionNameUtf8.Get();
@@ -805,8 +802,8 @@ TFuture<TDefaultErrorResult<FUpdateSessionJoinabilityImpl>> FSessionsEOSGS::Upda
 
 	// We get the active session info with the handle
 	EOS_ActiveSession_CopyInfoOptions CopyInfoOptions = {};
-	CopyInfoOptions.ApiVersion = EOS_ACTIVESESSION_COPYINFO_API_LATEST;
-	static_assert(EOS_ACTIVESESSION_COPYINFO_API_LATEST == 1, "EOS_ActiveSession_CopyInfoOptions updated, check new fields");
+	CopyInfoOptions.ApiVersion = 1;
+	UE_EOS_CHECK_API_MISMATCH(EOS_ACTIVESESSION_COPYINFO_API_LATEST, 1);
 
 	EOS_ActiveSession_Info* ActiveSessionInfo = nullptr;
 	EOS_EResult CopyInfoResult = EOS_ActiveSession_CopyInfo(ActiveSessionHandle, &CopyInfoOptions, &ActiveSessionInfo);
@@ -822,8 +819,8 @@ TFuture<TDefaultErrorResult<FUpdateSessionJoinabilityImpl>> FSessionsEOSGS::Upda
 		if (ActiveSessionInfo->State == EOS_EOnlineSessionState::EOS_OSS_InProgress)
 		{
 			EOS_Sessions_EndSessionOptions Options = {};
-			Options.ApiVersion = EOS_SESSIONS_ENDSESSION_API_LATEST;
-			static_assert(EOS_SESSIONS_ENDSESSION_API_LATEST == 1, "EOS_Sessions_EndSessionOptions updated, check new fields");
+			Options.ApiVersion = 1;
+			UE_EOS_CHECK_API_MISMATCH(EOS_SESSIONS_ENDSESSION_API_LATEST, 1);
 
 			Options.SessionName = SessionNameUtf8.Get();
 
@@ -851,8 +848,8 @@ TFuture<TDefaultErrorResult<FUpdateSessionJoinabilityImpl>> FSessionsEOSGS::Upda
 		if (ActiveSessionInfo->State == EOS_EOnlineSessionState::EOS_OSS_Pending || ActiveSessionInfo->State == EOS_EOnlineSessionState::EOS_OSS_Ended)
 		{
 			EOS_Sessions_StartSessionOptions Options = {};
-			Options.ApiVersion = EOS_SESSIONS_STARTSESSION_API_LATEST;
-			static_assert(EOS_SESSIONS_STARTSESSION_API_LATEST == 1, "EOS_Sessions_StartSessionOptions updated, check new fields");
+			Options.ApiVersion = 1;
+			UE_EOS_CHECK_API_MISMATCH(EOS_SESSIONS_STARTSESSION_API_LATEST, 1);
 
 			Options.SessionName = SessionNameUtf8.Get();
 
@@ -884,8 +881,8 @@ TOnlineAsyncOpHandle<FSendSingleSessionInviteImpl> FSessionsEOSGS::SendSingleSes
 	const FSendSingleSessionInviteImpl::Params& OpParams = Op->GetParams();
 
 	EOS_Sessions_SendInviteOptions SendInviteOptions = { };
-	SendInviteOptions.ApiVersion = EOS_SESSIONS_SENDINVITE_API_LATEST;
-	static_assert(EOS_SESSIONS_SENDINVITE_API_LATEST == 1, "EOS_Sessions_SendInviteOptions updated, check new fields");
+	SendInviteOptions.ApiVersion = 1;
+	UE_EOS_CHECK_API_MISMATCH(EOS_SESSIONS_SENDINVITE_API_LATEST, 1);
 
 	SendInviteOptions.LocalUserId = GetProductUserIdChecked(OpParams.LocalAccountId);
 
@@ -928,8 +925,8 @@ TFuture<TOnlineResult<FLeaveSession>> FSessionsEOSGS::LeaveSessionImpl(const FLe
 	TFuture<TOnlineResult<FLeaveSession>> Future = Promise.GetFuture();
 
 	EOS_Sessions_DestroySessionOptions DestroySessionOptions = {};
-	DestroySessionOptions.ApiVersion = EOS_SESSIONS_DESTROYSESSION_API_LATEST;
-	static_assert(EOS_SESSIONS_DESTROYSESSION_API_LATEST == 1, "EOS_Sessions_DestroySessionOptions updated, check new fields");
+	DestroySessionOptions.ApiVersion = 1;
+	UE_EOS_CHECK_API_MISMATCH(EOS_SESSIONS_DESTROYSESSION_API_LATEST, 1);
 
 	const FTCHARToUTF8 SessionNameUtf8(*Params.SessionName.ToString());
 	DestroySessionOptions.SessionName = SessionNameUtf8.Get();
@@ -969,8 +966,8 @@ TFuture<TOnlineResult<FLeaveSession>> FSessionsEOSGS::LeaveSessionImpl(const FLe
 void FSessionsEOSGS::SetSessionSearchMaxResults(FSessionSearchHandleEOSGS& SessionSearchHandle, uint32 MaxResults)
 {
 	EOS_SessionSearch_SetMaxResultsOptions Options = { };
-	Options.ApiVersion = EOS_SESSIONSEARCH_SETMAXSEARCHRESULTS_API_LATEST;
-	static_assert(EOS_SESSIONSEARCH_SETMAXSEARCHRESULTS_API_LATEST == 1, "EOS_SessionSearch_SetMaxResultsOptions updated, check new fields");
+	Options.ApiVersion = 1;
+	UE_EOS_CHECK_API_MISMATCH(EOS_SESSIONSEARCH_SETMAXSEARCHRESULTS_API_LATEST, 1);
 
 	// We truncate the max results parameter up to the value of EOS_SESSIONS_MAX_SEARCH_RESULTS
 	Options.MaxSearchResults = FMath::Clamp(MaxResults, 1, EOS_SESSIONS_MAX_SEARCH_RESULTS);
@@ -987,8 +984,8 @@ void FSessionsEOSGS::SetSessionSearchParameters(FSessionSearchHandleEOSGS& Sessi
 	for (const FFindSessionsSearchFilter& Filter : Filters)
 	{
 		EOS_SessionSearch_SetParameterOptions Options = { };
-		Options.ApiVersion = EOS_SESSIONSEARCH_SETPARAMETER_API_LATEST;
-		static_assert(EOS_SESSIONSEARCH_SETPARAMETER_API_LATEST == 1, "EOS_SessionSearch_SetMaxResultsOptions updated, check new fields");
+		Options.ApiVersion = 1;
+		UE_EOS_CHECK_API_MISMATCH(EOS_SESSIONSEARCH_SETPARAMETER_API_LATEST, 1);
 
 		Options.ComparisonOp = ToServiceType(Filter.ComparisonOp);
 		
@@ -1006,8 +1003,8 @@ void FSessionsEOSGS::SetSessionSearchParameters(FSessionSearchHandleEOSGS& Sessi
 void FSessionsEOSGS::SetSessionSearchSessionId(FSessionSearchHandleEOSGS& SessionSearchHandle, const FOnlineSessionId& SessionId)
 {
 	EOS_SessionSearch_SetSessionIdOptions Options = { };
-	Options.ApiVersion = EOS_SESSIONSEARCH_SETSESSIONID_API_LATEST;
-	static_assert(EOS_SESSIONSEARCH_SETSESSIONID_API_LATEST == 1, "EOS_SessionSearch_SetSessionIdOptions updated, check new fields");
+	Options.ApiVersion = 1;
+	UE_EOS_CHECK_API_MISMATCH(EOS_SESSIONSEARCH_SETSESSIONID_API_LATEST, 1);
 
 	const FTCHARToUTF8 SessionIdUtf8(*FOnlineSessionIdRegistryEOSGS::Get().ToLogString(SessionId));
 	Options.SessionId = SessionIdUtf8.Get();
@@ -1022,8 +1019,8 @@ void FSessionsEOSGS::SetSessionSearchSessionId(FSessionSearchHandleEOSGS& Sessio
 void FSessionsEOSGS::SetSessionSearchTargetId(FSessionSearchHandleEOSGS& SessionSearchHandle, const FAccountId& TargetAccountId)
 {
 	EOS_SessionSearch_SetTargetUserIdOptions Options = { };
-	Options.ApiVersion = EOS_SESSIONSEARCH_SETTARGETUSERID_API_LATEST;
-	static_assert(EOS_SESSIONSEARCH_SETTARGETUSERID_API_LATEST == 1, "EOS_SessionSearch_SetTargetUserIdOptions updated, check new fields");
+	Options.ApiVersion = 1;
+	UE_EOS_CHECK_API_MISMATCH(EOS_SESSIONSEARCH_SETTARGETUSERID_API_LATEST, 1);
 	Options.TargetUserId = GetProductUserIdChecked(TargetAccountId);
 
 	EOS_EResult ResultCode = EOS_SessionSearch_SetTargetUserId(SessionSearchHandle.SearchHandle, &Options);
@@ -1077,8 +1074,8 @@ TFuture<TOnlineResult<FFindSessions>> FSessionsEOSGS::FindSessionsImpl(const FFi
 
 	// We start preparing the search
 	EOS_Sessions_CreateSessionSearchOptions CreateSessionSearchOptions = { };
-	CreateSessionSearchOptions.ApiVersion = EOS_SESSIONS_CREATESESSIONSEARCH_API_LATEST;
-	static_assert(EOS_SESSIONS_CREATESESSIONSEARCH_API_LATEST == 1, "EOS_Sessions_CreateSessionSearchOptions updated, check new fields");
+	CreateSessionSearchOptions.ApiVersion = 1;
+	UE_EOS_CHECK_API_MISMATCH(EOS_SESSIONS_CREATESESSIONSEARCH_API_LATEST, 1);
 	CreateSessionSearchOptions.MaxSearchResults = FMath::Clamp(Params.MaxResults, 0, EOS_SESSIONS_MAX_SEARCH_RESULTS);
 
 	EOS_HSessionSearch SearchHandle = nullptr;
@@ -1101,8 +1098,8 @@ TFuture<TOnlineResult<FFindSessions>> FSessionsEOSGS::FindSessionsImpl(const FFi
 	WriteSessionSearchHandle(*CurrentSessionSearchHandleEOSGS, Params);
 
 	EOS_SessionSearch_FindOptions FindOptions = { };
-	FindOptions.ApiVersion = EOS_SESSIONSEARCH_FIND_API_LATEST;
-	static_assert(EOS_SESSIONSEARCH_FIND_API_LATEST == 2, "EOS_SessionSearch_FindOptions updated, check new fields");
+	FindOptions.ApiVersion = 2;
+	UE_EOS_CHECK_API_MISMATCH(EOS_SESSIONSEARCH_FIND_API_LATEST, 2);
 	FindOptions.LocalUserId = GetProductUserIdChecked(Params.LocalAccountId);
 
 	EOS_Async(EOS_SessionSearch_Find, CurrentSessionSearchHandleEOSGS->SearchHandle, FindOptions,
@@ -1121,13 +1118,13 @@ TFuture<TOnlineResult<FFindSessions>> FSessionsEOSGS::FindSessionsImpl(const FFi
 
 			// For a successful session, we'll get the search results
 			EOS_SessionSearch_GetSearchResultCountOptions GetSearchResultCountOptions = { };
-			GetSearchResultCountOptions.ApiVersion = EOS_SESSIONSEARCH_GETSEARCHRESULTCOUNT_API_LATEST;
-			static_assert(EOS_SESSIONSEARCH_GETSEARCHRESULTCOUNT_API_LATEST == 1, "EOS_SessionSearch_GetSearchResultCountOptions updated, check new fields");
+			GetSearchResultCountOptions.ApiVersion = 1;
+			UE_EOS_CHECK_API_MISMATCH(EOS_SESSIONSEARCH_GETSEARCHRESULTCOUNT_API_LATEST, 1);
 			int32 NumSearchResults = EOS_SessionSearch_GetSearchResultCount(CurrentSessionSearchHandleEOSGS->SearchHandle, &GetSearchResultCountOptions);
 
 			EOS_SessionSearch_CopySearchResultByIndexOptions CopySearchResultByIndexOptions = { };
-			CopySearchResultByIndexOptions.ApiVersion = EOS_SESSIONSEARCH_COPYSEARCHRESULTBYINDEX_API_LATEST;
-			static_assert(EOS_SESSIONSEARCH_COPYSEARCHRESULTBYINDEX_API_LATEST == 1, "EOS_SessionSearch_CopySearchResultByIndexOptions updated, check new fields");
+			CopySearchResultByIndexOptions.ApiVersion = 1;
+			UE_EOS_CHECK_API_MISMATCH(EOS_SESSIONSEARCH_COPYSEARCHRESULTBYINDEX_API_LATEST, 1);
 
 			TArray<TFuture<TDefaultErrorResult<FBuildSessionFromDetailsHandle>>> PendingSessionsBuilt;
 
@@ -1222,8 +1219,8 @@ TFuture<TOnlineResult<FJoinSession>> FSessionsEOSGS::JoinSessionImpl(const FJoin
 
 	// We start setup for the API call
 	EOS_Sessions_JoinSessionOptions JoinSessionOptions = { };
-	JoinSessionOptions.ApiVersion = EOS_SESSIONS_JOINSESSION_API_LATEST;
-	static_assert(EOS_SESSIONS_JOINSESSION_API_LATEST == 2, "EOS_Sessions_JoinSessionOptions updated, check new fields");
+	JoinSessionOptions.ApiVersion = 2;
+	UE_EOS_CHECK_API_MISMATCH(EOS_SESSIONS_JOINSESSION_API_LATEST, 2);
 
 	JoinSessionOptions.bPresenceEnabled = Params.bPresenceEnabled;
 
@@ -1309,8 +1306,8 @@ TOptional<FOnlineError> FSessionsEOSGS::CheckState(const FJoinSession::Params& P
 TOnlineAsyncOpHandle<FBuildSessionFromDetailsHandle> FSessionsEOSGS::BuildSessionFromInvite(const FAccountId& LocalAccountId, const FString& InInviteId)
 {
 	EOS_Sessions_CopySessionHandleByInviteIdOptions CopySessionHandleByInviteIdOptions = { };
-	CopySessionHandleByInviteIdOptions.ApiVersion = EOS_SESSIONS_COPYSESSIONHANDLEBYINVITEID_API_LATEST;
-	static_assert(EOS_SESSIONS_COPYSESSIONHANDLEBYINVITEID_API_LATEST == 1, "EOS_Sessions_CopySessionHandleByInviteIdOptions updated, check new fields");
+	CopySessionHandleByInviteIdOptions.ApiVersion = 1;
+	UE_EOS_CHECK_API_MISMATCH(EOS_SESSIONS_COPYSESSIONHANDLEBYINVITEID_API_LATEST, 1);
 
 	const FTCHARToUTF8 InviteIdUtf8(*InInviteId);
 	CopySessionHandleByInviteIdOptions.InviteId = InviteIdUtf8.Get();
@@ -1334,8 +1331,8 @@ TOnlineAsyncOpHandle<FBuildSessionFromDetailsHandle> FSessionsEOSGS::BuildSessio
 TOnlineAsyncOpHandle<FBuildSessionFromDetailsHandle> FSessionsEOSGS::BuildSessionFromUIEvent(const FAccountId& LocalAccountId, const EOS_UI_EventId& UIEventId)
 {
 	EOS_Sessions_CopySessionHandleByUiEventIdOptions CopySessionHandleByUiEventIdOptions = { };
-	CopySessionHandleByUiEventIdOptions.ApiVersion = EOS_SESSIONS_COPYSESSIONHANDLEBYUIEVENTID_API_LATEST;
-	static_assert(EOS_SESSIONS_COPYSESSIONHANDLEBYUIEVENTID_API_LATEST == 1, "EOS_Sessions_CopySessionHandleByUiEventIdOptions updated, check new fields");
+	CopySessionHandleByUiEventIdOptions.ApiVersion = 1;
+	UE_EOS_CHECK_API_MISMATCH(EOS_SESSIONS_COPYSESSIONHANDLEBYUIEVENTID_API_LATEST, 1);
 
 	CopySessionHandleByUiEventIdOptions.UiEventId = UIEventId;
 
@@ -1362,8 +1359,8 @@ TResult<TArray<EOS_ProductUserId>, FOnlineError> GetProductUserIdsFromEOSGSSessi
 	// This code will remain commented until the related functionality is made available by EOSSDK
 
 	/*EOS_SessionDetails_CopyInfoOptions CopyInfoOptions = {};
-	CopyInfoOptions.ApiVersion = EOS_SESSIONDETAILS_COPYINFO_API_LATEST;
-	static_assert(EOS_SESSIONDETAILS_COPYINFO_API_LATEST == 1, "EOS_SessionDetails_CopyInfoOptions updated, check new fields");
+	CopyInfoOptions.ApiVersion = 1;
+	UE_EOS_CHECK_API_MISMATCH(EOS_SESSIONDETAILS_COPYINFO_API_LATEST, 1);
 
 	EOS_SessionDetails_Info* SessionDetailsInfo = nullptr;
 	EOS_EResult CopyInfoResult = EOS_SessionDetails_CopyInfo(SessionDetailsHandle, &CopyInfoOptions, &SessionDetailsInfo);
@@ -1463,8 +1460,8 @@ TFuture<TOnlineResult<FRejectSessionInvite>> FSessionsEOSGS::RejectSessionInvite
 	TFuture<TOnlineResult<FRejectSessionInvite>> Future = Promise.GetFuture();
 
 	EOS_Sessions_RejectInviteOptions RejectInviteOptions = { };
-	RejectInviteOptions.ApiVersion = EOS_SESSIONS_REJECTINVITE_API_LATEST;
-	static_assert(EOS_SESSIONS_REJECTINVITE_API_LATEST == 1, "EOS_Sessions_RejectInviteOptions updated, check new fields");
+	RejectInviteOptions.ApiVersion = 1;
+	UE_EOS_CHECK_API_MISMATCH(EOS_SESSIONS_REJECTINVITE_API_LATEST, 1);
 
 	RejectInviteOptions.LocalUserId = GetProductUserIdChecked(Params.LocalAccountId);
 
@@ -1513,8 +1510,8 @@ TFuture<TOnlineResult<FAddSessionMember>> FSessionsEOSGS::AddSessionMemberImpl(c
 				if (!ResolvedAccountIds.IsEmpty())
 				{
 					EOS_Sessions_RegisterPlayersOptions Options = { };
-					Options.ApiVersion = EOS_SESSIONS_REGISTERPLAYERS_API_LATEST;
-					static_assert(EOS_SESSIONS_REGISTERPLAYERS_API_LATEST == 2, "EOS_Sessions_RegisterPlayersOptions updated, check new fields");
+					Options.ApiVersion = 2;
+					UE_EOS_CHECK_API_MISMATCH(EOS_SESSIONS_REGISTERPLAYERS_API_LATEST, 2);
 
 					Options.PlayersToRegisterCount = 1;
 
@@ -1566,8 +1563,8 @@ TFuture<TOnlineResult<FRemoveSessionMember>> FSessionsEOSGS::RemoveSessionMember
 				if (!ResolvedAccountIds.IsEmpty())
 				{
 					EOS_Sessions_UnregisterPlayersOptions Options = { };
-					Options.ApiVersion = EOS_SESSIONS_UNREGISTERPLAYERS_API_LATEST;
-					static_assert(EOS_SESSIONS_UNREGISTERPLAYERS_API_LATEST == 2, "EOS_Sessions_UnregisterPlayersOptions updated, check new fields");
+					Options.ApiVersion = 2;
+					UE_EOS_CHECK_API_MISMATCH(EOS_SESSIONS_UNREGISTERPLAYERS_API_LATEST, 2);
 
 					Options.PlayersToUnregisterCount = 1;
 
diff --git a/Engine/Plugins/Online/OnlineServicesEOSGS/Source/Private/Online/SessionsEOSGSTypes.cpp b/Engine/Plugins/Online/OnlineServicesEOSGS/Source/Private/Online/SessionsEOSGSTypes.cpp
--- a/Engine/Plugins/Online/OnlineServicesEOSGS/Source/Private/Online/SessionsEOSGSTypes.cpp
+++ b/Engine/Plugins/Online/OnlineServicesEOSGS/Source/Private/Online/SessionsEOSGSTypes.cpp
@@ -2,6 +2,8 @@
 
 #include "Online/SessionsEOSGSTypes.h"
 
+#include "EOSShared.h"
+
 namespace UE::Online {
 
 EOS_EOnlineSessionPermissionLevel ToServiceType(const ESessionJoinPolicy& Value)
@@ -82,8 +84,8 @@ FSessionAttributeConverter<ESessionAttributeConversionType::ToService>::FSession
 FSessionAttributeConverter<ESessionAttributeConversionType::ToService>::FSessionAttributeConverter(const TPair<FSchemaAttributeId, FSchemaVariant>& InData)
 	: KeyConverterStorage(*InData.Key.ToString())
 {
-	AttributeData.ApiVersion = EOS_SESSIONS_ATTRIBUTEDATA_API_LATEST;
-	static_assert(EOS_SESSIONS_ATTRIBUTEDATA_API_LATEST == 1, "EOS_Sessions_AttributeData updated, check new fields");
+	AttributeData.ApiVersion = 1;
+	UE_EOS_CHECK_API_MISMATCH(EOS_SESSIONS_ATTRIBUTEDATA_API_LATEST, 1);
 
 	AttributeData.Key = KeyConverterStorage.Get();
 
diff --git a/Engine/Plugins/Online/OnlineServicesEOSGS/Source/Private/Online/StatsEOSGS.cpp b/Engine/Plugins/Online/OnlineServicesEOSGS/Source/Private/Online/StatsEOSGS.cpp
--- a/Engine/Plugins/Online/OnlineServicesEOSGS/Source/Private/Online/StatsEOSGS.cpp
+++ b/Engine/Plugins/Online/OnlineServicesEOSGS/Source/Private/Online/StatsEOSGS.cpp
@@ -57,8 +57,8 @@ TOnlineAsyncOpHandle<FUpdateStats> FStatsEOSGS::UpdateStats(FUpdateStats::Params
 			for (const TPair<FString, FStatValue>& StatPair : UpdateUserStats.Stats)
 			{
 				EOS_Stats_IngestData& EOSStat = EOSData[Index];
-				EOSStat.ApiVersion = EOS_STATS_INGESTDATA_API_LATEST;
-				static_assert(EOS_STATS_INGESTDATA_API_LATEST == 1, "EOS_Stats_IngestData updated, check new fields");
+				EOSStat.ApiVersion = 1;
+				UE_EOS_CHECK_API_MISMATCH(EOS_STATS_INGESTDATA_API_LATEST, 1);
 
 				EOSStat.IngestAmount = static_cast<int32>(StatPair.Value.GetInt64());
 				FCStringAnsi::Strncpy(EOSStatNames[Index].StatName, TCHAR_TO_UTF8(*StatPair.Key), STAT_NAME_MAX_LENGTH_EOS);
@@ -68,8 +68,8 @@ TOnlineAsyncOpHandle<FUpdateStats> FStatsEOSGS::UpdateStats(FUpdateStats::Params
 			}
 
 			EOS_Stats_IngestStatOptions Options = { };
-			Options.ApiVersion = EOS_STATS_INGESTSTAT_API_LATEST;
-			static_assert(EOS_STATS_INGESTSTAT_API_LATEST == 3, "EOS_Stats_IngestStat updated, check new fields");
+			Options.ApiVersion = 3;
+			UE_EOS_CHECK_API_MISMATCH(EOS_STATS_INGESTSTAT_API_LATEST, 3);
 			Options.LocalUserId = GetProductUserIdChecked(InAsyncOp.GetParams().LocalAccountId);
 			Options.TargetUserId = GetProductUserIdChecked(UpdateUserStats.AccountId);
 			Options.Stats = EOSData.GetData();
@@ -105,8 +105,8 @@ namespace Private
 void QueryStatsEOS(EOS_HStats StatsHandle, const FAccountId& LocalAccountId, const FAccountId& TargetAccountId, const TArray<FString>& StatNames, TPromise<const EOS_Stats_OnQueryStatsCompleteCallbackInfo*>&& Promise)
 {
 	EOS_Stats_QueryStatsOptions Options;
-	Options.ApiVersion = EOS_STATS_QUERYSTATS_API_LATEST;
-	static_assert(EOS_STATS_INGESTSTAT_API_LATEST == 3, "EOS_Stats_QueryStats updated, check new fields");
+	Options.ApiVersion = 3;
+	UE_EOS_CHECK_API_MISMATCH(EOS_STATS_INGESTSTAT_API_LATEST, 3);
 	Options.StartTime = EOS_STATS_TIME_UNDEFINED;
 	Options.EndTime = EOS_STATS_TIME_UNDEFINED;
 
@@ -138,8 +138,8 @@ void ReadStatsFromEOSResult(EOS_HStats StatsHandle, const EOS_Stats_OnQueryStats
 	{
 		char StatNameANSI[STAT_NAME_MAX_LENGTH_EOS];
 		EOS_Stats_CopyStatByNameOptions Options = { };
-		Options.ApiVersion = EOS_STATS_COPYSTATBYNAME_API_LATEST;
-		static_assert(EOS_STATS_COPYSTATBYNAME_API_LATEST == 1, "EOS_Stats_CopyStatByNameOptions updated, check new fields");
+		Options.ApiVersion = 1;
+		UE_EOS_CHECK_API_MISMATCH(EOS_STATS_COPYSTATBYNAME_API_LATEST, 1);
 		Options.TargetUserId = Data->TargetUserId;
 		Options.Name = StatNameANSI;
 		FCStringAnsi::Strncpy(StatNameANSI, TCHAR_TO_UTF8(*StatName), STAT_NAME_MAX_LENGTH_EOS);
diff --git a/Engine/Plugins/Online/OnlineServicesEOSGS/Source/Private/Online/TitleFileEOSGS.cpp b/Engine/Plugins/Online/OnlineServicesEOSGS/Source/Private/Online/TitleFileEOSGS.cpp
--- a/Engine/Plugins/Online/OnlineServicesEOSGS/Source/Private/Online/TitleFileEOSGS.cpp
+++ b/Engine/Plugins/Online/OnlineServicesEOSGS/Source/Private/Online/TitleFileEOSGS.cpp
@@ -59,8 +59,8 @@ TOnlineAsyncOpHandle<FTitleFileEnumerateFiles> FTitleFileEOSGS::EnumerateFiles(F
 			const char* SearchTagPtr = Utf8SearchTag.Get();
 
 			EOS_TitleStorage_QueryFileListOptions Options = {};
-			Options.ApiVersion = EOS_TITLESTORAGE_QUERYFILELISTOPTIONS_API_LATEST;
-			static_assert(EOS_TITLESTORAGE_QUERYFILELISTOPTIONS_API_LATEST == 1, "EOS_TitleStorage_QueryFileListOptions updated, check new fields");
+			Options.ApiVersion = 1;
+			UE_EOS_CHECK_API_MISMATCH(EOS_TITLESTORAGE_QUERYFILELISTOPTIONS_API_LATEST, 1);
 			Options.LocalUserId = GetProductUserIdChecked(Params.LocalAccountId);
 			Options.ListOfTagsCount = 1;
 			Options.ListOfTags = &SearchTagPtr;
@@ -80,8 +80,8 @@ TOnlineAsyncOpHandle<FTitleFileEnumerateFiles> FTitleFileEOSGS::EnumerateFiles(F
 			const EOS_ProductUserId LocalUserPuid = GetProductUserIdChecked(Params.LocalAccountId);
 
 			EOS_TitleStorage_GetFileMetadataCountOptions GetCountOptions;
-			GetCountOptions.ApiVersion = EOS_TITLESTORAGE_GETFILEMETADATACOUNTOPTIONS_API_LATEST;
-			static_assert(EOS_TITLESTORAGE_GETFILEMETADATACOUNTOPTIONS_API_LATEST == 1, "EOS_TitleStorage_GetFileMetadataCountOptions updated, check new fields");
+			GetCountOptions.ApiVersion = 1;
+			UE_EOS_CHECK_API_MISMATCH(EOS_TITLESTORAGE_GETFILEMETADATACOUNTOPTIONS_API_LATEST, 1);
 			GetCountOptions.LocalUserId = LocalUserPuid;
 
 			const uint32_t NumFiles = EOS_TitleStorage_GetFileMetadataCount(TitleStorageHandle, &GetCountOptions);
@@ -92,8 +92,8 @@ TOnlineAsyncOpHandle<FTitleFileEnumerateFiles> FTitleFileEOSGS::EnumerateFiles(F
 			for (uint32_t FileIdx = 0; FileIdx < NumFiles; FileIdx++)
 			{
 				EOS_TitleStorage_CopyFileMetadataAtIndexOptions CopyOptions;
-				CopyOptions.ApiVersion = EOS_TITLESTORAGE_COPYFILEMETADATAATINDEXOPTIONS_API_LATEST;
-				static_assert(EOS_TITLESTORAGE_COPYFILEMETADATAATINDEXOPTIONS_API_LATEST == 1, "EOS_TitleStorage_CopyFileMetadataAtIndexOptions updated, check new fields");
+				CopyOptions.ApiVersion = 1;
+				UE_EOS_CHECK_API_MISMATCH(EOS_TITLESTORAGE_COPYFILEMETADATAATINDEXOPTIONS_API_LATEST, 1);
 				CopyOptions.LocalUserId = LocalUserPuid;
 				CopyOptions.Index = FileIdx;
 
@@ -182,8 +182,8 @@ TOnlineAsyncOpHandle<FTitleFileReadFile> FTitleFileEOSGS::ReadFile(FTitleFileRea
 			const FTCHARToUTF8 Utf8Filename(*Params.Filename);
 
 			EOS_TitleStorage_ReadFileOptions Options = {};
-			Options.ApiVersion = EOS_TITLESTORAGE_READFILEOPTIONS_API_LATEST;
-			static_assert(EOS_TITLESTORAGE_READFILEOPTIONS_API_LATEST == 1, "EOS_TitleStorage_ReadFileOptions updated, check new fields");
+			Options.ApiVersion = 1;
+			UE_EOS_CHECK_API_MISMATCH(EOS_TITLESTORAGE_READFILEOPTIONS_API_LATEST, 1);
 			Options.LocalUserId = GetProductUserIdChecked(Params.LocalAccountId);
 			Options.Filename = Utf8Filename.Get();
 			Options.ReadChunkLengthBytes = Config.ReadChunkLengthBytes;
diff --git a/Engine/Plugins/Online/OnlineServicesEOSGS/Source/Private/Online/UserFileEOSGS.cpp b/Engine/Plugins/Online/OnlineServicesEOSGS/Source/Private/Online/UserFileEOSGS.cpp
--- a/Engine/Plugins/Online/OnlineServicesEOSGS/Source/Private/Online/UserFileEOSGS.cpp
+++ b/Engine/Plugins/Online/OnlineServicesEOSGS/Source/Private/Online/UserFileEOSGS.cpp
@@ -55,8 +55,8 @@ TOnlineAsyncOpHandle<FUserFileEnumerateFiles> FUserFileEOSGS::EnumerateFiles(FUs
 			const FUserFileEnumerateFiles::Params& Params = Op.GetParams();
 
 			EOS_PlayerDataStorage_QueryFileListOptions Options = {};
-			Options.ApiVersion = EOS_PLAYERDATASTORAGE_QUERYFILELISTOPTIONS_API_LATEST;
-			static_assert(EOS_PLAYERDATASTORAGE_QUERYFILELISTOPTIONS_API_LATEST == 1, "EOS_PlayerDataStorage_QueryFileListOptions updated, check new fields");
+			Options.ApiVersion = 1;
+			UE_EOS_CHECK_API_MISMATCH(EOS_PLAYERDATASTORAGE_QUERYFILELISTOPTIONS_API_LATEST, 1);
 			Options.LocalUserId = GetProductUserIdChecked(Params.LocalAccountId);
 
 			EOS_Async(EOS_PlayerDataStorage_QueryFileList, PlayerDataStorageHandle, Options, MoveTemp(Promise));
@@ -76,8 +76,8 @@ TOnlineAsyncOpHandle<FUserFileEnumerateFiles> FUserFileEOSGS::EnumerateFiles(FUs
 			const EOS_ProductUserId LocalUserPuid = GetProductUserIdChecked(Params.LocalAccountId);
 
 			EOS_PlayerDataStorage_GetFileMetadataCountOptions GetCountOptions;
-			GetCountOptions.ApiVersion = EOS_PLAYERDATASTORAGE_GETFILEMETADATACOUNTOPTIONS_API_LATEST;
-			static_assert(EOS_PLAYERDATASTORAGE_GETFILEMETADATACOUNTOPTIONS_API_LATEST == 1, "EOS_PlayerDataStorage_GetFileMetadataCountOptions updated, check new fields");
+			GetCountOptions.ApiVersion = 1;
+			UE_EOS_CHECK_API_MISMATCH(EOS_PLAYERDATASTORAGE_GETFILEMETADATACOUNTOPTIONS_API_LATEST, 1);
 			GetCountOptions.LocalUserId = LocalUserPuid;
 
 			int32_t NumFiles;
@@ -97,8 +97,8 @@ TOnlineAsyncOpHandle<FUserFileEnumerateFiles> FUserFileEOSGS::EnumerateFiles(FUs
 			for (int32_t FileIdx = 0; FileIdx < NumFiles; FileIdx++)
 			{
 				EOS_PlayerDataStorage_CopyFileMetadataAtIndexOptions CopyOptions;
-				CopyOptions.ApiVersion = EOS_PLAYERDATASTORAGE_COPYFILEMETADATAATINDEXOPTIONS_API_LATEST;
-				static_assert(EOS_PLAYERDATASTORAGE_COPYFILEMETADATAATINDEXOPTIONS_API_LATEST == 1, "EOS_PlayerDataStorage_CopyFileMetadataAtIndexOptions updated, check new fields");
+				CopyOptions.ApiVersion = 1;
+				UE_EOS_CHECK_API_MISMATCH(EOS_PLAYERDATASTORAGE_COPYFILEMETADATAATINDEXOPTIONS_API_LATEST, 1);
 				CopyOptions.LocalUserId = LocalUserPuid;
 				CopyOptions.Index = FileIdx;
 
@@ -186,8 +186,8 @@ TOnlineAsyncOpHandle<FUserFileReadFile> FUserFileEOSGS::ReadFile(FUserFileReadFi
 			const auto Utf8Filename = StringCast<UTF8CHAR>(*Params.Filename);
 
 			EOS_PlayerDataStorage_ReadFileOptions Options = {};
-			Options.ApiVersion = EOS_PLAYERDATASTORAGE_READFILEOPTIONS_API_LATEST;
-			static_assert(EOS_PLAYERDATASTORAGE_READFILEOPTIONS_API_LATEST == 1, "EOS_PlayerDataStorage_ReadFileOptions updated, check new fields");
+			Options.ApiVersion = 1;
+			UE_EOS_CHECK_API_MISMATCH(EOS_PLAYERDATASTORAGE_READFILEOPTIONS_API_LATEST, 1);
 			Options.LocalUserId = GetProductUserIdChecked(Params.LocalAccountId);
 			Options.Filename = (const char*)Utf8Filename.Get();
 			Options.ReadChunkLengthBytes = Config.ChunkLengthBytes;
@@ -268,8 +268,8 @@ TOnlineAsyncOpHandle<FUserFileWriteFile> FUserFileEOSGS::WriteFile(FUserFileWrit
 		const auto Utf8Filename = StringCast<UTF8CHAR>(*Params.Filename);
 
 		EOS_PlayerDataStorage_WriteFileOptions Options = {};
-		Options.ApiVersion = EOS_PLAYERDATASTORAGE_WRITEFILEOPTIONS_API_LATEST;
-		static_assert(EOS_PLAYERDATASTORAGE_WRITEFILEOPTIONS_API_LATEST == 1, "EOS_PlayerDataStorage_WriteFileOptions updated, check new fields");
+		Options.ApiVersion = 1;
+		UE_EOS_CHECK_API_MISMATCH(EOS_PLAYERDATASTORAGE_WRITEFILEOPTIONS_API_LATEST, 1);
 		Options.LocalUserId = GetProductUserIdChecked(Params.LocalAccountId);
 		Options.Filename = (const char*)Utf8Filename.Get();
 		Options.ChunkLengthBytes = Config.ChunkLengthBytes;
@@ -328,8 +328,8 @@ TOnlineAsyncOpHandle<FUserFileCopyFile> FUserFileEOSGS::CopyFile(FUserFileCopyFi
 		const auto Utf8TargetFilename = StringCast<UTF8CHAR>(*Params.TargetFilename);
 
 		EOS_PlayerDataStorage_DuplicateFileOptions Options = {};
-		Options.ApiVersion = EOS_PLAYERDATASTORAGE_DUPLICATEFILEOPTIONS_API_LATEST;
-		static_assert(EOS_PLAYERDATASTORAGE_DUPLICATEFILEOPTIONS_API_LATEST == 1, "EOS_PlayerDataStorage_DuplicateFileOptions updated, check new fields");
+		Options.ApiVersion = 1;
+		UE_EOS_CHECK_API_MISMATCH(EOS_PLAYERDATASTORAGE_DUPLICATEFILEOPTIONS_API_LATEST, 1);
 		Options.LocalUserId = GetProductUserIdChecked(Params.LocalAccountId);
 		Options.SourceFilename = (const char*)Utf8SourceFilename.Get();
 		Options.DestinationFilename = (const char*)Utf8TargetFilename.Get();
@@ -377,8 +377,8 @@ TOnlineAsyncOpHandle<FUserFileDeleteFile> FUserFileEOSGS::DeleteFile(FUserFileDe
 		const auto Utf8Filename = StringCast<UTF8CHAR>(*Params.Filename);
 
 		EOS_PlayerDataStorage_DeleteFileOptions Options = {};
-		Options.ApiVersion = EOS_PLAYERDATASTORAGE_DELETEFILEOPTIONS_API_LATEST;
-		static_assert(EOS_PLAYERDATASTORAGE_DELETEFILEOPTIONS_API_LATEST == 1, "EOS_PlayerDataStorage_DeleteFileOptions updated, check new fields");
+		Options.ApiVersion = 1;
+		UE_EOS_CHECK_API_MISMATCH(EOS_PLAYERDATASTORAGE_DELETEFILEOPTIONS_API_LATEST, 1);
 		Options.LocalUserId = GetProductUserIdChecked(Params.LocalAccountId);
 		Options.Filename = (const char*)Utf8Filename.Get();
 
diff --git a/Engine/Plugins/Online/OnlineSubsystemEOS/Source/OnlineSubsystemEOS/Private/IOS/IOSEOSHelpers.cpp b/Engine/Plugins/Online/OnlineSubsystemEOS/Source/OnlineSubsystemEOS/Private/IOS/IOSEOSHelpers.cpp
--- a/Engine/Plugins/Online/OnlineSubsystemEOS/Source/OnlineSubsystemEOS/Private/IOS/IOSEOSHelpers.cpp
+++ b/Engine/Plugins/Online/OnlineSubsystemEOS/Source/OnlineSubsystemEOS/Private/IOS/IOSEOSHelpers.cpp
@@ -62,7 +62,8 @@ void FIOSEOSHelpers::PlatformAuthCredentials(EOS_Auth_Credentials &Credentials)
 	}
 #endif
 
-	CredentialsOptions.ApiVersion = EOS_IOS_AUTH_CREDENTIALSOPTIONS_API_LATEST;
+	CredentialsOptions.ApiVersion = 1;
+	UE_EOS_CHECK_API_MISMATCH(EOS_IOS_AUTH_CREDENTIALSOPTIONS_API_LATEST, 1);
 #if PRESENTATIONCONTEXT_REQUIRED
 	CredentialsOptions.PresentationContextProviding = (void*)CFBridgingRetain(PresentationContextProvider);		// SDK will release when consumed
 #else
diff --git a/Engine/Plugins/Online/OnlineSubsystemEOS/Source/OnlineSubsystemEOS/Private/OnlineAchievementsEOS.cpp b/Engine/Plugins/Online/OnlineSubsystemEOS/Source/OnlineSubsystemEOS/Private/OnlineAchievementsEOS.cpp
--- a/Engine/Plugins/Online/OnlineSubsystemEOS/Source/OnlineSubsystemEOS/Private/OnlineAchievementsEOS.cpp
+++ b/Engine/Plugins/Online/OnlineSubsystemEOS/Source/OnlineSubsystemEOS/Private/OnlineAchievementsEOS.cpp
@@ -40,13 +40,10 @@ void FOnlineAchievementsEOS::QueryAchievements(const FUniqueNetId& PlayerId, con
 	}
 
 	EOS_Achievements_QueryPlayerAchievementsOptions Options = { };
-	Options.ApiVersion = EOS_ACHIEVEMENTS_QUERYPLAYERACHIEVEMENTS_API_LATEST;
-#if EOS_ACHIEVEMENTS_QUERYPLAYERACHIEVEMENTS_API_LATEST >= 2
+	Options.ApiVersion = 2;
+	UE_EOS_CHECK_API_MISMATCH(EOS_ACHIEVEMENTS_QUERYPLAYERACHIEVEMENTS_API_LATEST, 2);
 	Options.LocalUserId = EOSSubsystem->UserManager->GetLocalProductUserId(LocalUserId);
 	Options.TargetUserId = Options.LocalUserId;
-#else
-	Options.UserId = EOSSubsystem->UserManager->GetLocalProductUserId(LocalUserId);
-#endif
 
 	FQueryProgressCallback* CallbackObj = new FQueryProgressCallback(FOnlineAchievementsEOSWeakPtr(AsShared()));
 	CallbackObj->CallbackLambda = [this, LambdaPlayerId = PlayerId.AsShared(), OnComplete = FOnQueryAchievementsCompleteDelegate(Delegate)](const EOS_Achievements_OnQueryPlayerAchievementsCompleteCallbackInfo* Data)
@@ -61,18 +58,16 @@ void FOnlineAchievementsEOS::QueryAchievements(const FUniqueNetId& PlayerId, con
 			EOS_ProductUserId UserId = EOSSubsystem->UserManager->GetLocalProductUserId(LocalUserNum);
 
 			EOS_Achievements_GetPlayerAchievementCountOptions CountOptions = { };
-			CountOptions.ApiVersion = EOS_ACHIEVEMENTS_GETPLAYERACHIEVEMENTCOUNT_API_LATEST;
+			CountOptions.ApiVersion = 1;
+			UE_EOS_CHECK_API_MISMATCH(EOS_ACHIEVEMENTS_GETPLAYERACHIEVEMENTCOUNT_API_LATEST, 1);
 			CountOptions.UserId = UserId;
 			uint32 Count = EOS_Achievements_GetPlayerAchievementCount(EOSSubsystem->AchievementsHandle, &CountOptions);
 
 			EOS_Achievements_CopyPlayerAchievementByIndexOptions CopyOptions = { };
-			CopyOptions.ApiVersion = EOS_ACHIEVEMENTS_COPYPLAYERACHIEVEMENTBYINDEX_API_LATEST;
-#if EOS_ACHIEVEMENTS_COPYPLAYERACHIEVEMENTBYINDEX_API_LATEST >= 2
+			CopyOptions.ApiVersion = 2;
+			UE_EOS_CHECK_API_MISMATCH(EOS_ACHIEVEMENTS_COPYPLAYERACHIEVEMENTBYINDEX_API_LATEST, 2);
 			CopyOptions.LocalUserId = UserId;
 			CopyOptions.TargetUserId = UserId;
-#else
-			CopyOptions.UserId = UserId;
-#endif
 
 			for (uint32 Index = 0; Index < Count; Index++)
 			{
@@ -129,7 +124,8 @@ void FOnlineAchievementsEOS::QueryAchievementDescriptions(const FUniqueNetId& Pl
 	}
 
 	EOS_Achievements_QueryDefinitionsOptions Options = { };
-	Options.ApiVersion = EOS_ACHIEVEMENTS_QUERYDEFINITIONS_API_LATEST;
+	Options.ApiVersion = 3;
+	UE_EOS_CHECK_API_MISMATCH(EOS_ACHIEVEMENTS_QUERYDEFINITIONS_API_LATEST, 3);
 	Options.LocalUserId = EOSSubsystem->UserManager->GetLocalProductUserId(LocalUserId);
 
 	FQueryDefinitionsCallback* CallbackObj = new FQueryDefinitionsCallback(FOnlineAchievementsEOSWeakPtr(AsShared()));
@@ -139,11 +135,13 @@ void FOnlineAchievementsEOS::QueryAchievementDescriptions(const FUniqueNetId& Pl
 		if (bWasSuccessful)
 		{
 			EOS_Achievements_GetAchievementDefinitionCountOptions CountOptions = { };
-			CountOptions.ApiVersion = EOS_ACHIEVEMENTS_GETACHIEVEMENTDEFINITIONCOUNT_API_LATEST;
+			CountOptions.ApiVersion = 1;
+			UE_EOS_CHECK_API_MISMATCH(EOS_ACHIEVEMENTS_GETACHIEVEMENTDEFINITIONCOUNT_API_LATEST, 1);
 			uint32 Count = EOS_Achievements_GetAchievementDefinitionCount(EOSSubsystem->AchievementsHandle, &CountOptions);
 
 			EOS_Achievements_CopyAchievementDefinitionByIndexOptions CopyOptions = { };
-			CopyOptions.ApiVersion = EOS_ACHIEVEMENTS_COPYDEFINITIONBYINDEX_API_LATEST;
+			CopyOptions.ApiVersion = 1;
+			UE_EOS_CHECK_API_MISMATCH(EOS_ACHIEVEMENTS_COPYDEFINITIONBYINDEX_API_LATEST, 1);
 			CachedAchievementDefinitions.Empty(Count);
 			CachedAchievementDefinitionsMap.Empty();
 
diff --git a/Engine/Plugins/Online/OnlineSubsystemEOS/Source/OnlineSubsystemEOS/Private/OnlineLeaderboardsEOS.cpp b/Engine/Plugins/Online/OnlineSubsystemEOS/Source/OnlineSubsystemEOS/Private/OnlineLeaderboardsEOS.cpp
--- a/Engine/Plugins/Online/OnlineSubsystemEOS/Source/OnlineSubsystemEOS/Private/OnlineLeaderboardsEOS.cpp
+++ b/Engine/Plugins/Online/OnlineSubsystemEOS/Source/OnlineSubsystemEOS/Private/OnlineLeaderboardsEOS.cpp
@@ -23,7 +23,8 @@ struct FQueryLeaderboardForUserOptions :
 		: EOS_Leaderboards_QueryLeaderboardUserScoresOptions()
 		, ProductUserIds(MoveTemp(InProductUserIds))
 	{
-		ApiVersion = EOS_LEADERBOARDS_QUERYLEADERBOARDUSERSCORES_API_LATEST;
+		ApiVersion = 2;
+		UE_EOS_CHECK_API_MISMATCH(EOS_LEADERBOARDS_QUERYLEADERBOARDUSERSCORES_API_LATEST, 2);
 
 		StartTime = EOS_LEADERBOARDS_TIME_UNDEFINED;
 		EndTime = EOS_LEADERBOARDS_TIME_UNDEFINED;
@@ -108,7 +109,8 @@ bool FOnlineLeaderboardsEOS::ReadLeaderboards(const TArray<FUniqueNetIdRef>& Pla
 	{
 		FCStringAnsi::Strncpy(Options.PointerArray[Index], TCHAR_TO_UTF8(*Column.ColumnName.ToString()), EOS_OSS_STRING_BUFFER_LENGTH);
 		EOS_Leaderboards_UserScoresQueryStatInfo& StatInfo = Options.StatInfoArray[Index];
-		StatInfo.ApiVersion = EOS_LEADERBOARDS_USERSCORESQUERYSTATINFO_API_LATEST;
+		StatInfo.ApiVersion = 1;
+		UE_EOS_CHECK_API_MISMATCH(EOS_LEADERBOARDS_USERSCORESQUERYSTATINFO_API_LATEST, 1);
 		StatInfo.StatName = Options.PointerArray[Index];
 		StatInfo.Aggregation = EOS_ELeaderboardAggregation::EOS_LA_Latest;
 		Index++;
@@ -145,7 +147,8 @@ bool FOnlineLeaderboardsEOS::ReadLeaderboards(const TArray<FUniqueNetIdRef>& Pla
 			FOnlineStatsRow Row(Nickname, NetId);
 
 			EOS_Leaderboards_CopyLeaderboardUserScoreByUserIdOptions UserCopyOptions = { };
-			UserCopyOptions.ApiVersion = EOS_LEADERBOARDS_COPYLEADERBOARDUSERSCOREBYUSERID_API_LATEST;
+			UserCopyOptions.ApiVersion = 1;
+			UE_EOS_CHECK_API_MISMATCH(EOS_LEADERBOARDS_COPYLEADERBOARDUSERSCOREBYUSERID_API_LATEST, 1);
 			UserCopyOptions.UserId = UserId;
 			UserCopyOptions.StatName = StatName;
 
@@ -248,7 +251,8 @@ bool FOnlineLeaderboardsEOS::ReadLeaderboardsAroundRank(int32 Rank, uint32 Range
 
 	char LeaderboardId[EOS_OSS_STRING_BUFFER_LENGTH];
 	EOS_Leaderboards_QueryLeaderboardRanksOptions Options = { };
-	Options.ApiVersion = EOS_LEADERBOARDS_QUERYLEADERBOARDRANKS_API_LATEST;
+	Options.ApiVersion = 2;
+	UE_EOS_CHECK_API_MISMATCH(EOS_LEADERBOARDS_QUERYLEADERBOARDRANKS_API_LATEST, 2);
 	Options.LeaderboardId = LeaderboardId;
 	Options.LocalUserId = EOSSubsystem->UserManager->GetLocalProductUserId(0);
 	FCStringAnsi::Strncpy(LeaderboardId, TCHAR_TO_UTF8(*ReadObject->LeaderboardName.ToString()), EOS_OSS_STRING_BUFFER_LENGTH);
@@ -266,7 +270,8 @@ bool FOnlineLeaderboardsEOS::ReadLeaderboardsAroundRank(int32 Rank, uint32 Range
 			return;
 		}
 		EOS_Leaderboards_GetLeaderboardRecordCountOptions CountOptions = { };
-		CountOptions.ApiVersion = EOS_LEADERBOARDS_GETLEADERBOARDRECORDCOUNT_API_LATEST;
+		CountOptions.ApiVersion = 1;
+		UE_EOS_CHECK_API_MISMATCH(EOS_LEADERBOARDS_GETLEADERBOARDRECORDCOUNT_API_LATEST, 1);
 
 		uint32 LeaderboardCount = EOS_Leaderboards_GetLeaderboardRecordCount(EOSSubsystem->LeaderboardsHandle, &CountOptions);
 		// Handle fewer entries than our start index
@@ -280,10 +285,12 @@ bool FOnlineLeaderboardsEOS::ReadLeaderboardsAroundRank(int32 Rank, uint32 Range
 		uint32 NewEndIndex = FMath::Clamp<uint32>(EndIndex, 0, LeaderboardCount - 1);
 
 		EOS_Leaderboards_CopyLeaderboardRecordByIndexOptions CopyOptions = { };
-		CopyOptions.ApiVersion = EOS_LEADERBOARDS_COPYLEADERBOARDRECORDBYINDEX_API_LATEST;
+		CopyOptions.ApiVersion = 2;
+		UE_EOS_CHECK_API_MISMATCH(EOS_LEADERBOARDS_COPYLEADERBOARDRECORDBYINDEX_API_LATEST, 2);
 
 		EOS_Connect_GetProductUserIdMappingOptions Options = { };
-		Options.ApiVersion = EOS_CONNECT_GETPRODUCTUSERIDMAPPING_API_LATEST;
+		Options.ApiVersion = 1;
+		UE_EOS_CHECK_API_MISMATCH(EOS_CONNECT_GETPRODUCTUSERIDMAPPING_API_LATEST, 1);
 		Options.AccountIdType = EOS_EExternalAccountType::EOS_EAT_EPIC;
 		Options.LocalUserId = EOSSubsystem->UserManager->GetLocalProductUserId(0);
 
diff --git a/Engine/Plugins/Online/OnlineSubsystemEOS/Source/OnlineSubsystemEOS/Private/OnlineSessionEOS.cpp b/Engine/Plugins/Online/OnlineSubsystemEOS/Source/OnlineSubsystemEOS/Private/OnlineSessionEOS.cpp
--- a/Engine/Plugins/Online/OnlineSubsystemEOS/Source/OnlineSubsystemEOS/Private/OnlineSessionEOS.cpp
+++ b/Engine/Plugins/Online/OnlineSubsystemEOS/Source/OnlineSubsystemEOS/Private/OnlineSessionEOS.cpp
@@ -118,7 +118,8 @@ struct FAttributeOptions :
 	FAttributeOptions(const char* InKey, const char* InValue) :
 		EOS_Sessions_AttributeData()
 	{
-		ApiVersion = EOS_SESSIONS_SESSIONATTRIBUTEDATA_API_LATEST;
+		ApiVersion = 1;
+		UE_EOS_CHECK_API_MISMATCH(EOS_SESSIONS_SESSIONATTRIBUTEDATA_API_LATEST, 1);
 		ValueType = EOS_ESessionAttributeType::EOS_SAT_String;
 		Value.AsUtf8 = ValueAnsi;
 		Key = KeyAnsi;
@@ -129,7 +130,8 @@ struct FAttributeOptions :
 	FAttributeOptions(const char* InKey, bool InValue) :
 		EOS_Sessions_AttributeData()
 	{
-		ApiVersion = EOS_SESSIONS_SESSIONATTRIBUTEDATA_API_LATEST;
+		ApiVersion = 1;
+		UE_EOS_CHECK_API_MISMATCH(EOS_SESSIONS_SESSIONATTRIBUTEDATA_API_LATEST, 1);
 		ValueType = EOS_ESessionAttributeType::EOS_SAT_Boolean;
 		Value.AsBool = InValue ? EOS_TRUE : EOS_FALSE;
 		Key = KeyAnsi;
@@ -139,7 +141,8 @@ struct FAttributeOptions :
 	FAttributeOptions(const char* InKey, float InValue) :
 		EOS_Sessions_AttributeData()
 	{
-		ApiVersion = EOS_SESSIONS_SESSIONATTRIBUTEDATA_API_LATEST;
+		ApiVersion = 1;
+		UE_EOS_CHECK_API_MISMATCH(EOS_SESSIONS_SESSIONATTRIBUTEDATA_API_LATEST, 1);
 		ValueType = EOS_ESessionAttributeType::EOS_SAT_Double;
 		Value.AsDouble = InValue;
 		Key = KeyAnsi;
@@ -149,7 +152,8 @@ struct FAttributeOptions :
 	FAttributeOptions(const char* InKey, int32 InValue) :
 		EOS_Sessions_AttributeData()
 	{
-		ApiVersion = EOS_SESSIONS_SESSIONATTRIBUTEDATA_API_LATEST;
+		ApiVersion = 1;
+		UE_EOS_CHECK_API_MISMATCH(EOS_SESSIONS_SESSIONATTRIBUTEDATA_API_LATEST, 1);
 		ValueType = EOS_ESessionAttributeType::EOS_SAT_Int64;
 		Value.AsInt64 = InValue;
 		Key = KeyAnsi;
@@ -159,7 +163,8 @@ struct FAttributeOptions :
 	FAttributeOptions(const char* InKey, const FVariantData& InValue) :
 		EOS_Sessions_AttributeData()
 	{
-		ApiVersion = EOS_SESSIONS_SESSIONATTRIBUTEDATA_API_LATEST;
+		ApiVersion = 1;
+		UE_EOS_CHECK_API_MISMATCH(EOS_SESSIONS_SESSIONATTRIBUTEDATA_API_LATEST, 1);
 
 		switch (InValue.GetType())
 		{
@@ -236,7 +241,8 @@ struct FLobbyAttributeOptions :
 	FLobbyAttributeOptions(const char* InKey, const char* InValue) :
 		EOS_Lobby_AttributeData()
 	{
-		ApiVersion = EOS_LOBBY_ATTRIBUTEDATA_API_LATEST;
+		ApiVersion = 1;
+		UE_EOS_CHECK_API_MISMATCH(EOS_LOBBY_ATTRIBUTEDATA_API_LATEST, 1);
 		ValueType = EOS_ELobbyAttributeType::EOS_SAT_String;
 		Value.AsUtf8 = ValueAnsi;
 		Key = KeyAnsi;
@@ -247,7 +253,8 @@ struct FLobbyAttributeOptions :
 	FLobbyAttributeOptions(const char* InKey, bool InValue) :
 		EOS_Lobby_AttributeData()
 	{
-		ApiVersion = EOS_LOBBY_ATTRIBUTEDATA_API_LATEST;
+		ApiVersion = 1;
+		UE_EOS_CHECK_API_MISMATCH(EOS_LOBBY_ATTRIBUTEDATA_API_LATEST, 1);
 		ValueType = EOS_ELobbyAttributeType::EOS_SAT_Boolean;
 		Value.AsBool = InValue ? EOS_TRUE : EOS_FALSE;
 		Key = KeyAnsi;
@@ -257,7 +264,8 @@ struct FLobbyAttributeOptions :
 	FLobbyAttributeOptions(const char* InKey, float InValue) :
 		EOS_Lobby_AttributeData()
 	{
-		ApiVersion = EOS_LOBBY_ATTRIBUTEDATA_API_LATEST;
+		ApiVersion = 1;
+		UE_EOS_CHECK_API_MISMATCH(EOS_LOBBY_ATTRIBUTEDATA_API_LATEST, 1);
 		ValueType = EOS_ELobbyAttributeType::EOS_SAT_Double;
 		Value.AsDouble = InValue;
 		Key = KeyAnsi;
@@ -267,7 +275,8 @@ struct FLobbyAttributeOptions :
 	FLobbyAttributeOptions(const char* InKey, int32 InValue) :
 		EOS_Lobby_AttributeData()
 	{
-		ApiVersion = EOS_LOBBY_ATTRIBUTEDATA_API_LATEST;
+		ApiVersion = 1;
+		UE_EOS_CHECK_API_MISMATCH(EOS_LOBBY_ATTRIBUTEDATA_API_LATEST, 1);
 		ValueType = EOS_ELobbyAttributeType::EOS_SAT_Int64;
 		Value.AsInt64 = InValue;
 		Key = KeyAnsi;
@@ -277,7 +286,8 @@ struct FLobbyAttributeOptions :
 	FLobbyAttributeOptions(const char* InKey, const FVariantData& InValue) :
 		EOS_Lobby_AttributeData()
 	{
-		ApiVersion = EOS_LOBBY_ATTRIBUTEDATA_API_LATEST;
+		ApiVersion = 1;
+		UE_EOS_CHECK_API_MISMATCH(EOS_LOBBY_ATTRIBUTEDATA_API_LATEST, 1);
 
 		switch (InValue.GetType())
 		{
@@ -463,7 +473,8 @@ void FOnlineSessionEOS::Init(const FString& InBucketId)
 		int32 LocalUserNum = EOSSubsystem->UserManager->GetLocalUserNumFromUniqueNetId(*NetId);
 
 		EOS_Sessions_CopySessionHandleByInviteIdOptions Options = { };
-		Options.ApiVersion = EOS_SESSIONS_COPYSESSIONHANDLEBYINVITEID_API_LATEST;
+		Options.ApiVersion = 1;
+		UE_EOS_CHECK_API_MISMATCH(EOS_SESSIONS_COPYSESSIONHANDLEBYINVITEID_API_LATEST, 1);
 		Options.InviteId = Data->InviteId;
 		EOS_HSessionDetails SessionDetails = nullptr;
 		EOS_EResult Result = EOS_Sessions_CopySessionHandleByInviteId(EOSSubsystem->SessionsHandle, &Options, &SessionDetails);
@@ -480,7 +491,8 @@ void FOnlineSessionEOS::Init(const FString& InBucketId)
 		}
 	};
 	EOS_Sessions_AddNotifySessionInviteAcceptedOptions Options = { };
-	Options.ApiVersion = EOS_SESSIONS_ADDNOTIFYSESSIONINVITEACCEPTED_API_LATEST;
+	Options.ApiVersion = 1;
+	UE_EOS_CHECK_API_MISMATCH(EOS_SESSIONS_ADDNOTIFYSESSIONINVITEACCEPTED_API_LATEST, 1);
 	SessionInviteAcceptedId = EOS_Sessions_AddNotifySessionInviteAccepted(EOSSubsystem->SessionsHandle, &Options, SessionInviteAcceptedCallbackObj, SessionInviteAcceptedCallbackObj->GetCallbackPtr());
 
 	// Lobbies
@@ -609,7 +621,8 @@ void FOnlineSessionEOS::RegisterLobbyNotifications()
 {
 	// Lobby data updates
 	EOS_Lobby_AddNotifyLobbyUpdateReceivedOptions AddNotifyLobbyUpdateReceivedOptions = { 0 };
-	AddNotifyLobbyUpdateReceivedOptions.ApiVersion = EOS_LOBBY_ADDNOTIFYLOBBYUPDATERECEIVED_API_LATEST;
+	AddNotifyLobbyUpdateReceivedOptions.ApiVersion = 1;
+	UE_EOS_CHECK_API_MISMATCH(EOS_LOBBY_ADDNOTIFYLOBBYUPDATERECEIVED_API_LATEST, 1);
 
 	FLobbyUpdateReceivedCallback* LobbyUpdateReceivedCallbackObj = new FLobbyUpdateReceivedCallback(FOnlineSessionEOSWeakPtr(AsShared()));
 	LobbyUpdateReceivedCallback = LobbyUpdateReceivedCallbackObj;
@@ -622,7 +635,8 @@ void FOnlineSessionEOS::RegisterLobbyNotifications()
 
 	// Lobby member data updates
 	EOS_Lobby_AddNotifyLobbyMemberUpdateReceivedOptions AddNotifyLobbyMemberUpdateReceivedOptions = { 0 };
-	AddNotifyLobbyMemberUpdateReceivedOptions.ApiVersion = EOS_LOBBY_ADDNOTIFYLOBBYMEMBERUPDATERECEIVED_API_LATEST;
+	AddNotifyLobbyMemberUpdateReceivedOptions.ApiVersion = 1;
+	UE_EOS_CHECK_API_MISMATCH(EOS_LOBBY_ADDNOTIFYLOBBYMEMBERUPDATERECEIVED_API_LATEST, 1);
 
 	FLobbyMemberUpdateReceivedCallback* LobbyMemberUpdateReceivedCallbackObj = new FLobbyMemberUpdateReceivedCallback(FOnlineSessionEOSWeakPtr(AsShared()));
 	LobbyMemberUpdateReceivedCallback = LobbyMemberUpdateReceivedCallbackObj;
@@ -635,7 +649,8 @@ void FOnlineSessionEOS::RegisterLobbyNotifications()
 
 	// Lobby member status updates (joined/left/disconnected/kicked/promoted)
 	EOS_Lobby_AddNotifyLobbyMemberStatusReceivedOptions AddNotifyLobbyMemberStatusReceivedOptions = { 0 };
-	AddNotifyLobbyMemberStatusReceivedOptions.ApiVersion = EOS_LOBBY_ADDNOTIFYLOBBYMEMBERSTATUSRECEIVED_API_LATEST;
+	AddNotifyLobbyMemberStatusReceivedOptions.ApiVersion = 1;
+	UE_EOS_CHECK_API_MISMATCH(EOS_LOBBY_ADDNOTIFYLOBBYMEMBERSTATUSRECEIVED_API_LATEST, 1);
 
 	FLobbyMemberStatusReceivedCallback* LobbyMemberStatusReceivedCallbackObj = new FLobbyMemberStatusReceivedCallback(FOnlineSessionEOSWeakPtr(AsShared()));
 	LobbyMemberStatusReceivedCallback = LobbyMemberStatusReceivedCallbackObj;
@@ -650,7 +665,8 @@ void FOnlineSessionEOS::RegisterLobbyNotifications()
 
 	// Accepted lobby invite notifications
 	EOS_Lobby_AddNotifyLobbyInviteAcceptedOptions AddNotifyLobbyInviteAcceptedOptions = { 0 };
-	AddNotifyLobbyInviteAcceptedOptions.ApiVersion = EOS_LOBBY_ADDNOTIFYLOBBYINVITEACCEPTED_API_LATEST;
+	AddNotifyLobbyInviteAcceptedOptions.ApiVersion = 1;
+	UE_EOS_CHECK_API_MISMATCH(EOS_LOBBY_ADDNOTIFYLOBBYINVITEACCEPTED_API_LATEST, 1);
 
 	FLobbyInviteAcceptedCallback* LobbyInviteAcceptedCallbackObj = new FLobbyInviteAcceptedCallback(FOnlineSessionEOSWeakPtr(AsShared()));
 	LobbyInviteAcceptedCallback = LobbyInviteAcceptedCallbackObj;
@@ -663,7 +679,8 @@ void FOnlineSessionEOS::RegisterLobbyNotifications()
 
 	// Accepted lobby join notifications
 	EOS_Lobby_AddNotifyJoinLobbyAcceptedOptions AddNotifyJoinLobbyAcceptedOptions = { 0 };
-	AddNotifyJoinLobbyAcceptedOptions.ApiVersion = EOS_LOBBY_ADDNOTIFYJOINLOBBYACCEPTED_API_LATEST;
+	AddNotifyJoinLobbyAcceptedOptions.ApiVersion = 1;
+	UE_EOS_CHECK_API_MISMATCH(EOS_LOBBY_ADDNOTIFYJOINLOBBYACCEPTED_API_LATEST, 1);
 
 	FJoinLobbyAcceptedCallback* JoinLobbyAcceptedCallbackObj = new FJoinLobbyAcceptedCallback(FOnlineSessionEOSWeakPtr(AsShared()));
 	JoinLobbyAcceptedCallback = JoinLobbyAcceptedCallbackObj;
@@ -682,7 +699,8 @@ void FOnlineSessionEOS::OnLobbyUpdateReceived(const EOS_LobbyId& LobbyId, const
 	if (Session)
 	{
 		EOS_Lobby_CopyLobbyDetailsHandleOptions Options = {};
-		Options.ApiVersion = EOS_LOBBY_COPYLOBBYDETAILSHANDLE_API_LATEST;
+		Options.ApiVersion = 1;
+		UE_EOS_CHECK_API_MISMATCH(EOS_LOBBY_COPYLOBBYDETAILSHANDLE_API_LATEST, 1);
 		Options.LobbyId = LobbyId;
 		Options.LocalUserId = LocalUserId;
 
@@ -695,7 +713,8 @@ void FOnlineSessionEOS::OnLobbyUpdateReceived(const EOS_LobbyId& LobbyId, const
 
 			EOS_LobbyDetails_Info* LobbyDetailsInfo = nullptr;
 			EOS_LobbyDetails_CopyInfoOptions CopyOptions = { };
-			CopyOptions.ApiVersion = EOS_LOBBYDETAILS_COPYINFO_API_LATEST;
+			CopyOptions.ApiVersion = 1;
+			UE_EOS_CHECK_API_MISMATCH(EOS_LOBBYDETAILS_COPYINFO_API_LATEST, 1);
 
 			EOS_EResult CopyInfoResult = EOS_LobbyDetails_CopyInfo(LobbyDetails->LobbyDetailsHandle, &CopyOptions, &LobbyDetailsInfo);
 			if (CopyInfoResult == EOS_EResult::EOS_Success)
@@ -875,7 +894,8 @@ void FOnlineSessionEOS::OnLobbyInviteAccepted(const char* InviteId, const EOS_Pr
 	int32 LocalUserNum = EOSSubsystem->UserManager->GetLocalUserNumFromUniqueNetId(*NetId);
 
 	EOS_Lobby_CopyLobbyDetailsHandleByInviteIdOptions Options = { };
-	Options.ApiVersion = EOS_LOBBY_COPYLOBBYDETAILSHANDLEBYINVITEID_API_LATEST;
+	Options.ApiVersion = 1;
+	UE_EOS_CHECK_API_MISMATCH(EOS_LOBBY_COPYLOBBYDETAILSHANDLEBYINVITEID_API_LATEST, 1);
 	Options.InviteId = InviteId;
 
 	EOS_HLobbyDetails LobbyDetailsHandle;
@@ -911,7 +931,8 @@ void FOnlineSessionEOS::OnJoinLobbyAccepted(const EOS_ProductUserId& LocalUserId
 	int32 LocalUserNum = EOSSubsystem->UserManager->GetLocalUserNumFromUniqueNetId(*NetId);
 
 	EOS_Lobby_CopyLobbyDetailsHandleByUiEventIdOptions Options = { 0 };
-	Options.ApiVersion = EOS_LOBBY_COPYLOBBYDETAILSHANDLEBYUIEVENTID_API_LATEST;
+	Options.ApiVersion = 1;
+	UE_EOS_CHECK_API_MISMATCH(EOS_LOBBY_COPYLOBBYDETAILSHANDLEBYUIEVENTID_API_LATEST, 1);
 	Options.UiEventId = UiEventId;
 
 	EOS_HLobbyDetails LobbyDetailsHandle;
@@ -1058,7 +1079,8 @@ uint32 FOnlineSessionEOS::CreateLANSession(int32 HostingPlayerNum, FNamedOnlineS
 void FOnlineSessionEOS::SetPermissionLevel(EOS_HSessionModification SessionModHandle, FNamedOnlineSession* Session)
 {
 	EOS_SessionModification_SetPermissionLevelOptions Options = { };
-	Options.ApiVersion = EOS_SESSIONMODIFICATION_SETPERMISSIONLEVEL_API_LATEST;
+	Options.ApiVersion = 1;
+	UE_EOS_CHECK_API_MISMATCH(EOS_SESSIONMODIFICATION_SETPERMISSIONLEVEL_API_LATEST, 1);
 	if (Session->SessionSettings.NumPublicConnections > 0)
 	{
 		Options.PermissionLevel = EOS_EOnlineSessionPermissionLevel::EOS_OSPF_PublicAdvertised;
@@ -1084,7 +1106,8 @@ void FOnlineSessionEOS::SetPermissionLevel(EOS_HSessionModification SessionModHa
 void FOnlineSessionEOS::SetMaxPlayers(EOS_HSessionModification SessionModHandle, FNamedOnlineSession* Session)
 {
 	EOS_SessionModification_SetMaxPlayersOptions Options = { };
-	Options.ApiVersion = EOS_SESSIONMODIFICATION_SETMAXPLAYERS_API_LATEST;
+	Options.ApiVersion = 1;
+	UE_EOS_CHECK_API_MISMATCH(EOS_SESSIONMODIFICATION_SETMAXPLAYERS_API_LATEST, 1);
 	Options.MaxPlayers = Session->SessionSettings.NumPrivateConnections + Session->SessionSettings.NumPublicConnections;
 
 	UE_LOG_ONLINE_SESSION(Log, TEXT("EOS_SessionModification_SetMaxPlayers() set to (%d) for session (%s)"), Options.MaxPlayers, *Session->SessionName.ToString());
@@ -1099,7 +1122,8 @@ void FOnlineSessionEOS::SetMaxPlayers(EOS_HSessionModification SessionModHandle,
 void FOnlineSessionEOS::SetInvitesAllowed(EOS_HSessionModification SessionModHandle, FNamedOnlineSession* Session)
 {
 	EOS_SessionModification_SetInvitesAllowedOptions Options = { };
-	Options.ApiVersion = EOS_SESSIONMODIFICATION_SETINVITESALLOWED_API_LATEST;
+	Options.ApiVersion = 1;
+	UE_EOS_CHECK_API_MISMATCH(EOS_SESSIONMODIFICATION_SETINVITESALLOWED_API_LATEST, 1);
 	Options.bInvitesAllowed = Session->SessionSettings.bAllowInvites ? EOS_TRUE : EOS_FALSE;
 
 	UE_LOG_ONLINE_SESSION(Log, TEXT("EOS_SessionModification_SetInvitesAllowed() set to (%s) for session (%s)"), *LexToString(Options.bInvitesAllowed), *Session->SessionName.ToString());
@@ -1114,7 +1138,8 @@ void FOnlineSessionEOS::SetInvitesAllowed(EOS_HSessionModification SessionModHan
 void FOnlineSessionEOS::SetJoinInProgress(EOS_HSessionModification SessionModHandle, FNamedOnlineSession* Session)
 {
 	EOS_SessionModification_SetJoinInProgressAllowedOptions Options = { };
-	Options.ApiVersion = EOS_SESSIONMODIFICATION_SETJOININPROGRESSALLOWED_API_LATEST;
+	Options.ApiVersion = 1;
+	UE_EOS_CHECK_API_MISMATCH(EOS_SESSIONMODIFICATION_SETJOININPROGRESSALLOWED_API_LATEST, 1);
 	Options.bAllowJoinInProgress = Session->SessionSettings.bAllowJoinInProgress ? EOS_TRUE : EOS_FALSE;
 
 	UE_LOG_ONLINE_SESSION(Log, TEXT("EOS_SessionModification_SetJoinInProgressAllowed() set to (%s) for session (%s)"), *LexToString(Options.bAllowJoinInProgress), *Session->SessionName.ToString());
@@ -1129,7 +1154,8 @@ void FOnlineSessionEOS::SetJoinInProgress(EOS_HSessionModification SessionModHan
 void FOnlineSessionEOS::AddAttribute(EOS_HSessionModification SessionModHandle, const EOS_Sessions_AttributeData* Attribute)
 {
 	EOS_SessionModification_AddAttributeOptions Options = { };
-	Options.ApiVersion = EOS_SESSIONMODIFICATION_ADDATTRIBUTE_API_LATEST;
+	Options.ApiVersion = 1;
+	UE_EOS_CHECK_API_MISMATCH(EOS_SESSIONMODIFICATION_ADDATTRIBUTE_API_LATEST, 1);
 	Options.AdvertisementType = EOS_ESessionAttributeAdvertisementType::EOS_SAAT_Advertise;
 	Options.SessionAttribute = Attribute;
 
@@ -1195,7 +1221,8 @@ struct FBeginMetricsOptions :
 	FBeginMetricsOptions() :
 		EOS_Metrics_BeginPlayerSessionOptions()
 	{
-		ApiVersion = EOS_METRICS_BEGINPLAYERSESSION_API_LATEST;
+		ApiVersion = 1;
+		UE_EOS_CHECK_API_MISMATCH(EOS_METRICS_BEGINPLAYERSESSION_API_LATEST, 1);
 		GameSessionId = SessionIdAnsi;
 		DisplayName = DisplayNameAnsi;
 		ServerIp = ServerIpAnsi;
@@ -1246,7 +1273,8 @@ struct FSessionCreateOptions :
 	FSessionCreateOptions(const char* InSessionNameAnsi) :
 		TNamedSessionOptions<EOS_Sessions_CreateSessionModificationOptions>(InSessionNameAnsi)
 	{
-		ApiVersion = EOS_SESSIONS_CREATESESSIONMODIFICATION_API_LATEST;
+		ApiVersion = 4;
+		UE_EOS_CHECK_API_MISMATCH(EOS_SESSIONS_CREATESESSIONMODIFICATION_API_LATEST, 4);
 		BucketId = BucketIdAnsi;
 	}
 };
@@ -1287,7 +1315,8 @@ uint32 FOnlineSessionEOS::CreateEOSSession(int32 HostingPlayerNum, FNamedOnlineS
 		FCStringAnsi::Strncpy(HostAddrAnsi, TCHAR_TO_UTF8(*HostAddr), EOS_OSS_STRING_BUFFER_LENGTH);
 
 		EOS_SessionModification_SetHostAddressOptions HostOptions = { };
-		HostOptions.ApiVersion = EOS_SESSIONMODIFICATION_SETHOSTADDRESS_API_LATEST;
+		HostOptions.ApiVersion = 1;
+		UE_EOS_CHECK_API_MISMATCH(EOS_SESSIONMODIFICATION_SETHOSTADDRESS_API_LATEST, 1);
 		// Expect URLs to look like "EOS:PUID:SocketName:Channel" and channel can be optional
 		HostOptions.HostAddress = HostAddrAnsi;
 		EOS_EResult HostResult = EOS_SessionModification_SetHostAddress(SessionModHandle, &HostOptions);
@@ -1355,7 +1384,8 @@ uint32 FOnlineSessionEOS::SharedSessionUpdate(EOS_HSessionModification SessionMo
 
 	// Commit the session changes
 	EOS_Sessions_UpdateSessionOptions CreateOptions = { };
-	CreateOptions.ApiVersion = EOS_SESSIONS_UPDATESESSION_API_LATEST;
+	CreateOptions.ApiVersion = 1;
+	UE_EOS_CHECK_API_MISMATCH(EOS_SESSIONS_UPDATESESSION_API_LATEST, 1);
 	CreateOptions.SessionModificationHandle = SessionModHandle;
 	EOS_Sessions_UpdateSession(EOSSubsystem->SessionsHandle, &CreateOptions, Callback, Callback->GetCallbackPtr());
 
@@ -1426,7 +1456,8 @@ struct FSessionStartOptions :
 	FSessionStartOptions(const char* InSessionNameAnsi) :
 		TNamedSessionOptions<EOS_Sessions_StartSessionOptions>(InSessionNameAnsi)
 	{
-		ApiVersion = EOS_SESSIONS_STARTSESSION_API_LATEST;
+		ApiVersion = 1;
+		UE_EOS_CHECK_API_MISMATCH(EOS_SESSIONS_STARTSESSION_API_LATEST, 1);
 	}
 };
 
@@ -1534,7 +1565,8 @@ struct FSessionUpdateOptions :
 	FSessionUpdateOptions(const char* InSessionNameAnsi) :
 		TNamedSessionOptions<EOS_Sessions_UpdateSessionModificationOptions>(InSessionNameAnsi)
 	{
-		ApiVersion = EOS_SESSIONS_UPDATESESSIONMODIFICATION_API_LATEST;
+		ApiVersion = 1;
+		UE_EOS_CHECK_API_MISMATCH(EOS_SESSIONS_UPDATESESSIONMODIFICATION_API_LATEST, 1);
 	}
 };
 
@@ -1654,7 +1686,8 @@ struct FSessionEndOptions :
 	FSessionEndOptions(const char* InSessionNameAnsi) :
 		TNamedSessionOptions<EOS_Sessions_EndSessionOptions>(InSessionNameAnsi)
 	{
-		ApiVersion = EOS_SESSIONS_ENDSESSION_API_LATEST;
+		ApiVersion = 1;
+		UE_EOS_CHECK_API_MISMATCH(EOS_SESSIONS_ENDSESSION_API_LATEST, 1);
 	}
 };
 
@@ -1779,7 +1812,8 @@ struct FEndMetricsOptions :
 	FEndMetricsOptions() :
 		EOS_Metrics_EndPlayerSessionOptions()
 	{
-		ApiVersion = EOS_METRICS_ENDPLAYERSESSION_API_LATEST;
+		ApiVersion = 1;
+		UE_EOS_CHECK_API_MISMATCH(EOS_METRICS_ENDPLAYERSESSION_API_LATEST, 1);
 		AccountId.External = ExternalIdAnsi;
 	}
 };
@@ -1808,7 +1842,8 @@ struct FSessionDestroyOptions :
 	FSessionDestroyOptions(const char* InSessionNameAnsi) :
 		TNamedSessionOptions<EOS_Sessions_DestroySessionOptions>(InSessionNameAnsi)
 	{
-		ApiVersion = EOS_SESSIONS_DESTROYSESSION_API_LATEST;
+		ApiVersion = 1;
+		UE_EOS_CHECK_API_MISMATCH(EOS_SESSIONS_DESTROYSESSION_API_LATEST, 1);
 	}
 };
 
@@ -1947,7 +1982,8 @@ bool FOnlineSessionEOS::FindSessionById(const FUniqueNetId& SearchingUserId, con
 	// We create the search handle
 	EOS_HLobbySearch LobbySearchHandle;
 	EOS_Lobby_CreateLobbySearchOptions CreateLobbySearchOptions = { 0 };
-	CreateLobbySearchOptions.ApiVersion = EOS_LOBBY_CREATELOBBYSEARCH_API_LATEST;
+	CreateLobbySearchOptions.ApiVersion = 1;
+	UE_EOS_CHECK_API_MISMATCH(EOS_LOBBY_CREATELOBBYSEARCH_API_LATEST, 1);
 	CreateLobbySearchOptions.MaxResults = EOS_SESSIONS_MAX_SEARCH_RESULTS;
 
 	EOS_EResult CreateLobbySearchResult = EOS_Lobby_CreateLobbySearch(LobbyHandle, &CreateLobbySearchOptions, &LobbySearchHandle);
@@ -1956,7 +1992,8 @@ bool FOnlineSessionEOS::FindSessionById(const FUniqueNetId& SearchingUserId, con
 		const FTCHARToUTF8 Utf8LobbyId(*SessionId.ToString());
 		// Set the lobby id we want to use to find lobbies			
 		EOS_LobbySearch_SetLobbyIdOptions SetLobbyIdOptions = { 0 };
-		SetLobbyIdOptions.ApiVersion = EOS_LOBBYSEARCH_SETLOBBYID_API_LATEST;
+		SetLobbyIdOptions.ApiVersion = 1;
+		UE_EOS_CHECK_API_MISMATCH(EOS_LOBBYSEARCH_SETLOBBYID_API_LATEST, 1);
 		SetLobbyIdOptions.LobbyId = (EOS_LobbyId)Utf8LobbyId.Get();
 
 		EOS_LobbySearch_SetLobbyId(LobbySearchHandle, &SetLobbyIdOptions);
@@ -1990,7 +2027,8 @@ bool FOnlineSessionEOS::FindSessionById(const FUniqueNetId& SearchingUserId, con
 void FOnlineSessionEOS::AddSearchAttribute(EOS_HSessionSearch SearchHandle, const EOS_Sessions_AttributeData* Attribute, EOS_EOnlineComparisonOp ComparisonOp)
 {
 	EOS_SessionSearch_SetParameterOptions Options = { };
-	Options.ApiVersion = EOS_SESSIONSEARCH_SETPARAMETER_API_LATEST;
+	Options.ApiVersion = 1;
+	UE_EOS_CHECK_API_MISMATCH(EOS_SESSIONSEARCH_SETPARAMETER_API_LATEST, 1);
 	Options.Parameter = Attribute;
 	Options.ComparisonOp = ComparisonOp;
 
@@ -2004,7 +2042,8 @@ void FOnlineSessionEOS::AddSearchAttribute(EOS_HSessionSearch SearchHandle, cons
 void FOnlineSessionEOS::AddLobbySearchAttribute(EOS_HLobbySearch LobbySearchHandle, const EOS_Lobby_AttributeData* Attribute, EOS_EOnlineComparisonOp ComparisonOp)
 {
 	EOS_LobbySearch_SetParameterOptions Options = { };
-	Options.ApiVersion = EOS_LOBBYSEARCH_SETPARAMETER_API_LATEST;
+	Options.ApiVersion = 1;
+	UE_EOS_CHECK_API_MISMATCH(EOS_LOBBYSEARCH_SETPARAMETER_API_LATEST, 1);
 	Options.Parameter = Attribute;
 	Options.ComparisonOp = ComparisonOp;
 
@@ -2044,13 +2083,15 @@ void FOnlineSessionEOS::CopySearchResult(EOS_HSessionDetails SessionHandle, EOS_
 void FOnlineSessionEOS::CopyAttributes(EOS_HSessionDetails SessionHandle, FOnlineSession& OutSession)
 {
 	EOS_SessionDetails_GetSessionAttributeCountOptions CountOptions = { };
-	CountOptions.ApiVersion = EOS_SESSIONDETAILS_GETSESSIONATTRIBUTECOUNT_API_LATEST;
+	CountOptions.ApiVersion = 1;
+	UE_EOS_CHECK_API_MISMATCH(EOS_SESSIONDETAILS_GETSESSIONATTRIBUTECOUNT_API_LATEST, 1);
 	int32 Count = EOS_SessionDetails_GetSessionAttributeCount(SessionHandle, &CountOptions);
 
 	for (int32 Index = 0; Index < Count; Index++)
 	{
 		EOS_SessionDetails_CopySessionAttributeByIndexOptions AttrOptions = { };
-		AttrOptions.ApiVersion = EOS_SESSIONDETAILS_COPYSESSIONATTRIBUTEBYINDEX_API_LATEST;
+		AttrOptions.ApiVersion = 1;
+		UE_EOS_CHECK_API_MISMATCH(EOS_SESSIONDETAILS_COPYSESSIONATTRIBUTEBYINDEX_API_LATEST, 1);
 		AttrOptions.AttrIndex = Index;
 
 		EOS_SessionDetails_Attribute* Attribute = NULL;
@@ -2123,7 +2164,8 @@ void FOnlineSessionEOS::AddSearchResult(EOS_HSessionDetails SessionHandle, const
 {
 	EOS_SessionDetails_Info* SessionInfo = nullptr;
 	EOS_SessionDetails_CopyInfoOptions CopyOptions = { };
-	CopyOptions.ApiVersion = EOS_SESSIONDETAILS_COPYINFO_API_LATEST;
+	CopyOptions.ApiVersion = 1;
+	UE_EOS_CHECK_API_MISMATCH(EOS_SESSIONDETAILS_COPYINFO_API_LATEST, 1);
 	EOS_EResult CopyResult = EOS_SessionDetails_CopyInfo(SessionHandle, &CopyOptions, &SessionInfo);
 	if (CopyResult == EOS_EResult::EOS_Success)
 	{
@@ -2144,7 +2186,8 @@ uint32 FOnlineSessionEOS::FindEOSSession(int32 SearchingPlayerNum, const TShared
 {
 	EOS_HSessionSearch SearchHandle = nullptr;
 	EOS_Sessions_CreateSessionSearchOptions HandleOptions = { };
-	HandleOptions.ApiVersion = EOS_SESSIONS_CREATESESSIONSEARCH_API_LATEST;
+	HandleOptions.ApiVersion = 1;
+	UE_EOS_CHECK_API_MISMATCH(EOS_SESSIONS_CREATESESSIONSEARCH_API_LATEST, 1);
 	HandleOptions.MaxSearchResults = FMath::Clamp(SearchSettings->MaxSearchResults, 0, EOS_SESSIONS_MAX_SEARCH_RESULTS);
 
 	EOS_EResult ResultCode = EOS_Sessions_CreateSessionSearch(EOSSubsystem->SessionsHandle, &HandleOptions, &SearchHandle);
@@ -2188,11 +2231,13 @@ uint32 FOnlineSessionEOS::FindEOSSession(int32 SearchingPlayerNum, const TShared
 		if (bWasSuccessful)
 		{
 			EOS_SessionSearch_GetSearchResultCountOptions SearchResultOptions = { };
-			SearchResultOptions.ApiVersion = EOS_SESSIONSEARCH_GETSEARCHRESULTCOUNT_API_LATEST;
+			SearchResultOptions.ApiVersion = 1;
+			UE_EOS_CHECK_API_MISMATCH(EOS_SESSIONSEARCH_GETSEARCHRESULTCOUNT_API_LATEST, 1);
 			int32 NumSearchResults = EOS_SessionSearch_GetSearchResultCount(CurrentSearchHandle->SearchHandle, &SearchResultOptions);
 
 			EOS_SessionSearch_CopySearchResultByIndexOptions IndexOptions = { };
-			IndexOptions.ApiVersion = EOS_SESSIONSEARCH_COPYSEARCHRESULTBYINDEX_API_LATEST;
+			IndexOptions.ApiVersion = 1;
+			UE_EOS_CHECK_API_MISMATCH(EOS_SESSIONSEARCH_COPYSEARCHRESULTBYINDEX_API_LATEST, 1);
 			for (int32 Index = 0; Index < NumSearchResults; Index++)
 			{
 				EOS_HSessionDetails SessionHandle = nullptr;
@@ -2217,7 +2262,8 @@ uint32 FOnlineSessionEOS::FindEOSSession(int32 SearchingPlayerNum, const TShared
 
 	// Execute the search
 	EOS_SessionSearch_FindOptions Options = { };
-	Options.ApiVersion = EOS_SESSIONSEARCH_FIND_API_LATEST;
+	Options.ApiVersion = 2;
+	UE_EOS_CHECK_API_MISMATCH(EOS_SESSIONSEARCH_FIND_API_LATEST, 2);
 	Options.LocalUserId = EOSSubsystem->UserManager->GetLocalProductUserId(SearchingPlayerNum);
 	EOS_SessionSearch_Find(SearchHandle, &Options, CallbackObj, CallbackObj->GetCallbackPtr());
 
@@ -2228,7 +2274,8 @@ void FOnlineSessionEOS::FindEOSSessionById(int32 LocalUserNum, const FUniqueNetI
 {
 	EOS_HSessionSearch SearchHandle = nullptr;
 	EOS_Sessions_CreateSessionSearchOptions HandleOptions = { };
-	HandleOptions.ApiVersion = EOS_SESSIONS_CREATESESSIONSEARCH_API_LATEST;
+	HandleOptions.ApiVersion = 1;
+	UE_EOS_CHECK_API_MISMATCH(EOS_SESSIONS_CREATESESSIONSEARCH_API_LATEST, 1);
 	HandleOptions.MaxSearchResults = 1;
 
 	EOS_EResult ResultCode = EOS_Sessions_CreateSessionSearch(EOSSubsystem->SessionsHandle, &HandleOptions, &SearchHandle);
@@ -2241,7 +2288,8 @@ void FOnlineSessionEOS::FindEOSSessionById(int32 LocalUserNum, const FUniqueNetI
 
 	const FTCHARToUTF8 Utf8SessionId(*SessionId.ToString());
 	EOS_SessionSearch_SetSessionIdOptions Options = { };
-	Options.ApiVersion = EOS_SESSIONSEARCH_SETSESSIONID_API_LATEST;
+	Options.ApiVersion = 1;
+	UE_EOS_CHECK_API_MISMATCH(EOS_SESSIONSEARCH_SETSESSIONID_API_LATEST, 1);
 	Options.SessionId = Utf8SessionId.Get();
 	ResultCode = EOS_SessionSearch_SetSessionId(SearchHandle, &Options);
 	if (ResultCode != EOS_EResult::EOS_Success)
@@ -2264,11 +2312,13 @@ void FOnlineSessionEOS::FindEOSSessionById(int32 LocalUserNum, const FUniqueNetI
 		if (bWasSuccessful)
 		{
 			EOS_SessionSearch_GetSearchResultCountOptions SearchResultOptions = { };
-			SearchResultOptions.ApiVersion = EOS_SESSIONSEARCH_GETSEARCHRESULTCOUNT_API_LATEST;
+			SearchResultOptions.ApiVersion = 1;
+			UE_EOS_CHECK_API_MISMATCH(EOS_SESSIONSEARCH_GETSEARCHRESULTCOUNT_API_LATEST, 1);
 			int32 NumSearchResults = EOS_SessionSearch_GetSearchResultCount(CurrentSearchHandle->SearchHandle, &SearchResultOptions);
 
 			EOS_SessionSearch_CopySearchResultByIndexOptions IndexOptions = { };
-			IndexOptions.ApiVersion = EOS_SESSIONSEARCH_COPYSEARCHRESULTBYINDEX_API_LATEST;
+			IndexOptions.ApiVersion = 1;
+			UE_EOS_CHECK_API_MISMATCH(EOS_SESSIONSEARCH_COPYSEARCHRESULTBYINDEX_API_LATEST, 1);
 			for (int32 Index = 0; Index < NumSearchResults; Index++)
 			{
 				EOS_HSessionDetails SessionHandle = nullptr;
@@ -2291,7 +2341,8 @@ void FOnlineSessionEOS::FindEOSSessionById(int32 LocalUserNum, const FUniqueNetI
 	};
 
 	EOS_SessionSearch_FindOptions FindOptions = { };
-	FindOptions.ApiVersion = EOS_SESSIONSEARCH_FIND_API_LATEST;
+	FindOptions.ApiVersion = 2;
+	UE_EOS_CHECK_API_MISMATCH(EOS_SESSIONSEARCH_FIND_API_LATEST, 2);
 	FindOptions.LocalUserId = EOSSubsystem->UserManager->GetLocalProductUserId(LocalUserNum);
 
 	EOS_SessionSearch_Find(SearchHandle, &FindOptions, CallbackObj, CallbackObj->GetCallbackPtr());
@@ -2455,7 +2506,8 @@ struct FJoinSessionOptions :
 	FJoinSessionOptions(const char* InSessionNameAnsi) :
 		TNamedSessionOptions<EOS_Sessions_JoinSessionOptions>(InSessionNameAnsi)
 	{
-		ApiVersion = EOS_SESSIONS_JOINSESSION_API_LATEST;
+		ApiVersion = 2;
+		UE_EOS_CHECK_API_MISMATCH(EOS_SESSIONS_JOINSESSION_API_LATEST, 2);
 	}
 };
 
@@ -2553,7 +2605,8 @@ bool FOnlineSessionEOS::FindFriendSession(int32 LocalUserNum, const FUniqueNetId
 	// We create the search handle
 	EOS_HLobbySearch LobbySearchHandle;
 	EOS_Lobby_CreateLobbySearchOptions CreateLobbySearchOptions = { 0 };
-	CreateLobbySearchOptions.ApiVersion = EOS_LOBBY_CREATELOBBYSEARCH_API_LATEST;
+	CreateLobbySearchOptions.ApiVersion = 1;
+	UE_EOS_CHECK_API_MISMATCH(EOS_LOBBY_CREATELOBBYSEARCH_API_LATEST, 1);
 	CreateLobbySearchOptions.MaxResults = EOS_SESSIONS_MAX_SEARCH_RESULTS;
 
 	EOS_EResult CreateLobbySearchResult = EOS_Lobby_CreateLobbySearch(LobbyHandle, &CreateLobbySearchOptions, &LobbySearchHandle);
@@ -2563,7 +2616,8 @@ bool FOnlineSessionEOS::FindFriendSession(int32 LocalUserNum, const FUniqueNetId
 
 		// Set the user we wan to use to find lobbies
 		EOS_LobbySearch_SetTargetUserIdOptions SetTargetUserIdOptions = { 0 };
-		SetTargetUserIdOptions.ApiVersion = EOS_LOBBYSEARCH_SETTARGETUSERID_API_LATEST;
+		SetTargetUserIdOptions.ApiVersion = 1;
+		UE_EOS_CHECK_API_MISMATCH(EOS_LOBBYSEARCH_SETTARGETUSERID_API_LATEST, 1);
 		SetTargetUserIdOptions.TargetUserId = FriendEOSId.GetProductUserId();
 
 		// TODO: Using this as a search parameter only works if we use the owner's id (search for lobbies we're already in). Pending API fix so it works with other users too.
@@ -2615,7 +2669,8 @@ struct FSendSessionInviteOptions :
 	FSendSessionInviteOptions(const char* InSessionNameAnsi) :
 		TNamedSessionOptions<EOS_Sessions_SendInviteOptions>(InSessionNameAnsi)
 	{
-		ApiVersion = EOS_SESSIONS_SENDINVITE_API_LATEST;
+		ApiVersion = 1;
+		UE_EOS_CHECK_API_MISMATCH(EOS_SESSIONS_SENDINVITE_API_LATEST, 1);
 	}
 };
 
@@ -2648,7 +2703,8 @@ bool FOnlineSessionEOS::SendSessionInvite(FName SessionName, EOS_ProductUserId S
 bool FOnlineSessionEOS::SendLobbyInvite(FName SessionName, EOS_ProductUserId SenderId, EOS_ProductUserId ReceiverId)
 {
 	EOS_Lobby_SendInviteOptions SendInviteOptions = { 0 };
-	SendInviteOptions.ApiVersion = EOS_LOBBY_SENDINVITE_API_LATEST;
+	SendInviteOptions.ApiVersion = 1;
+	UE_EOS_CHECK_API_MISMATCH(EOS_LOBBY_SENDINVITE_API_LATEST, 1);
 	const FTCHARToUTF8 Utf8LobbyId(*GetNamedSession(SessionName)->SessionInfo->GetSessionId().ToString());
 	SendInviteOptions.LobbyId = (EOS_LobbyId)Utf8LobbyId.Get();
 	SendInviteOptions.LocalUserId = SenderId;
@@ -2877,7 +2933,8 @@ void FOnlineSessionEOS::UpdateOrAddLobbyMember(const FUniqueNetIdEOSLobbyRef& Lo
 			const FTCHARToUTF8 Utf8LobbyId(*LobbyNetId->ToString());
 
 			EOS_Lobby_CopyLobbyDetailsHandleOptions Options = {};
-			Options.ApiVersion = EOS_LOBBY_COPYLOBBYDETAILSHANDLE_API_LATEST;
+			Options.ApiVersion = 1;
+			UE_EOS_CHECK_API_MISMATCH(EOS_LOBBY_COPYLOBBYDETAILSHANDLE_API_LATEST, 1);
 			Options.LobbyId = (EOS_LobbyId)Utf8LobbyId.Get();
 			Options.LocalUserId = EOSSubsystem->UserManager->GetLocalProductUserId();
 
@@ -3023,7 +3080,8 @@ bool FOnlineSessionEOS::RegisterPlayers(FName SessionName, const TArray< FUnique
 		if (bRegisterEOS && EOSIds.Num() > 0)
 		{
 			EOS_Sessions_RegisterPlayersOptions Options = { };
-			Options.ApiVersion = EOS_SESSIONS_REGISTERPLAYERS_API_LATEST;
+			Options.ApiVersion = 2;
+			UE_EOS_CHECK_API_MISMATCH(EOS_SESSIONS_REGISTERPLAYERS_API_LATEST, 2);
 			Options.PlayersToRegister = EOSIds.GetData();
 			Options.PlayersToRegisterCount = EOSIds.Num();
 			const FTCHARToUTF8 Utf8SessionName(*SessionName.ToString());
@@ -3095,7 +3153,8 @@ bool FOnlineSessionEOS::UnregisterPlayers(FName SessionName, const TArray< FUniq
 		if (bUnregisterEOS && EOSIds.Num() > 0)
 		{
 			EOS_Sessions_UnregisterPlayersOptions Options = { };
-			Options.ApiVersion = EOS_SESSIONS_UNREGISTERPLAYERS_API_LATEST;
+			Options.ApiVersion = 2;
+			UE_EOS_CHECK_API_MISMATCH(EOS_SESSIONS_UNREGISTERPLAYERS_API_LATEST, 2);
 			Options.PlayersToUnregister = EOSIds.GetData();
 			Options.PlayersToUnregisterCount = EOSIds.Num();
 			const FTCHARToUTF8 Utf8SessionName(*SessionName.ToString());
@@ -3450,7 +3509,8 @@ void FOnlineSessionEOS::RemovePlayerFromSession(int32 LocalUserNum, FName Sessio
 		const FUniqueNetIdEOS& TargetPlayerEOSId = FUniqueNetIdEOS::Cast(TargetPlayerId);
 
 		EOS_Lobby_KickMemberOptions KickMemberOptions = {};
-		KickMemberOptions.ApiVersion = EOS_LOBBY_KICKMEMBER_API_LATEST;
+		KickMemberOptions.ApiVersion = 1;
+		UE_EOS_CHECK_API_MISMATCH(EOS_LOBBY_KICKMEMBER_API_LATEST, 1);
 		const FTCHARToUTF8 Utf8LobbyId(*Session->SessionInfo->GetSessionId().ToString());
 		KickMemberOptions.LobbyId = (EOS_LobbyId)Utf8LobbyId.Get();
 		KickMemberOptions.LocalUserId = EOSSubsystem->UserManager->GetLocalProductUserId(LocalUserNum);
@@ -3518,8 +3578,8 @@ FString FOnlineSessionEOS::GetVoiceChatRoomName(int32 LocalUserNum, const FName&
 	const auto LobbyIdUTF8 = StringCast<UTF8CHAR>(*LobbyId.ToString());	
 
 	EOS_Lobby_GetRTCRoomNameOptions GetRTCRoomNameOptions = {};
-	GetRTCRoomNameOptions.ApiVersion = EOS_LOBBY_GETRTCROOMNAME_API_LATEST;
-	static_assert(EOS_LOBBY_GETRTCROOMNAME_API_LATEST == 1, "EOS_Lobby_GetRTCRoomNameOptions updated, check new fields");
+	GetRTCRoomNameOptions.ApiVersion = 1;
+	UE_EOS_CHECK_API_MISMATCH(EOS_LOBBY_GETRTCROOMNAME_API_LATEST, 1);
 	GetRTCRoomNameOptions.LobbyId = (const char*)LobbyIdUTF8.Get();
 	GetRTCRoomNameOptions.LocalUserId = EOSSubsystem->UserManager->GetLocalProductUserId(LocalUserNum);
 
@@ -3576,7 +3636,8 @@ uint32 FOnlineSessionEOS::CreateLobbySession(int32 HostingPlayerNum, FNamedOnlin
 	Session->SessionSettings.Get(SETTING_HOST_MIGRATION, bUseHostMigration);
 
 	EOS_Lobby_CreateLobbyOptions CreateLobbyOptions = { 0 };
-	CreateLobbyOptions.ApiVersion = EOS_LOBBY_CREATELOBBY_API_LATEST;
+	CreateLobbyOptions.ApiVersion = 8;
+	UE_EOS_CHECK_API_MISMATCH(EOS_LOBBY_CREATELOBBY_API_LATEST, 8);
 	CreateLobbyOptions.LocalUserId = LocalProductUserId;
 	CreateLobbyOptions.MaxLobbyMembers = GetLobbyMaxMembersFromSessionSettings(Session->SessionSettings);
 	CreateLobbyOptions.PermissionLevel = GetLobbyPermissionLevelFromSessionSettings(Session->SessionSettings);
@@ -3673,7 +3734,8 @@ uint32 FOnlineSessionEOS::JoinLobbySession(int32 PlayerNum, FNamedOnlineSession*
 			
 			// We retrieve the cached LobbyDetailsHandle and we start the join operation
 			EOS_Lobby_JoinLobbyOptions JoinLobbyOptions = { 0 };
-			JoinLobbyOptions.ApiVersion = EOS_LOBBY_JOINLOBBY_API_LATEST;
+			JoinLobbyOptions.ApiVersion = 3;
+			UE_EOS_CHECK_API_MISMATCH(EOS_LOBBY_JOINLOBBY_API_LATEST, 3);
 			JoinLobbyOptions.LocalUserId = EOSSubsystem->UserManager->GetLocalProductUserId(PlayerNum);
 			JoinLobbyOptions.bPresenceEnabled = Session->SessionSettings.bUsesPresence;
 
@@ -3749,7 +3811,8 @@ void FOnlineSessionEOS::SetLobbyPermissionLevel(EOS_HLobbyModification LobbyModi
 	check(Session != nullptr);
 
 	EOS_LobbyModification_SetPermissionLevelOptions Options = { 0 };
-	Options.ApiVersion = EOS_SESSIONMODIFICATION_SETPERMISSIONLEVEL_API_LATEST;
+	Options.ApiVersion = 1;
+	UE_EOS_CHECK_API_MISMATCH(EOS_SESSIONMODIFICATION_SETPERMISSIONLEVEL_API_LATEST, 1);
 	Options.PermissionLevel = GetLobbyPermissionLevelFromSessionSettings(Session->SessionSettings);
 
 	EOS_EResult ResultCode = EOS_LobbyModification_SetPermissionLevel(LobbyModificationHandle, &Options);
@@ -3764,7 +3827,8 @@ void FOnlineSessionEOS::SetLobbyMaxMembers(EOS_HLobbyModification LobbyModificat
 	check(Session != nullptr);
 
 	EOS_LobbyModification_SetMaxMembersOptions Options = { };
-	Options.ApiVersion = EOS_LOBBYMODIFICATION_SETMAXMEMBERS_API_LATEST;
+	Options.ApiVersion = 1;
+	UE_EOS_CHECK_API_MISMATCH(EOS_LOBBYMODIFICATION_SETMAXMEMBERS_API_LATEST, 1);
 	Options.MaxMembers = GetLobbyMaxMembersFromSessionSettings(Session->SessionSettings);
 
 	EOS_EResult ResultCode = EOS_LobbyModification_SetMaxMembers(LobbyModificationHandle, &Options);
@@ -3777,7 +3841,8 @@ void FOnlineSessionEOS::SetLobbyMaxMembers(EOS_HLobbyModification LobbyModificat
 void FOnlineSessionEOS::AddLobbyAttribute(EOS_HLobbyModification LobbyModificationHandle, const EOS_Lobby_AttributeData* Attribute)
 {
 	EOS_LobbyModification_AddAttributeOptions Options = { };
-	Options.ApiVersion = EOS_LOBBYMODIFICATION_ADDATTRIBUTE_API_LATEST;
+	Options.ApiVersion = 1;
+	UE_EOS_CHECK_API_MISMATCH(EOS_LOBBYMODIFICATION_ADDATTRIBUTE_API_LATEST, 1);
 	Options.Visibility = EOS_ELobbyAttributeVisibility::EOS_LAT_PUBLIC;
 	Options.Attribute = Attribute;
 
@@ -3791,7 +3856,8 @@ void FOnlineSessionEOS::AddLobbyAttribute(EOS_HLobbyModification LobbyModificati
 void FOnlineSessionEOS::AddLobbyMemberAttribute(EOS_HLobbyModification LobbyModificationHandle, const EOS_Lobby_AttributeData* Attribute)
 {
 	EOS_LobbyModification_AddMemberAttributeOptions Options = { };
-	Options.ApiVersion = EOS_LOBBYMODIFICATION_ADDMEMBERATTRIBUTE_API_LATEST;
+	Options.ApiVersion = 1;
+	UE_EOS_CHECK_API_MISMATCH(EOS_LOBBYMODIFICATION_ADDMEMBERATTRIBUTE_API_LATEST, 1);
 	Options.Visibility = EOS_ELobbyAttributeVisibility::EOS_LAT_PUBLIC;
 	Options.Attribute = Attribute;
 
@@ -3889,7 +3955,8 @@ uint32 FOnlineSessionEOS::UpdateLobbySession(FNamedOnlineSession* Session)
 	else
 	{
 		EOS_Lobby_UpdateLobbyModificationOptions UpdateLobbyModificationOptions = { 0 };
-		UpdateLobbyModificationOptions.ApiVersion = EOS_LOBBY_UPDATELOBBYMODIFICATION_API_LATEST;
+		UpdateLobbyModificationOptions.ApiVersion = 1;
+		UE_EOS_CHECK_API_MISMATCH(EOS_LOBBY_UPDATELOBBYMODIFICATION_API_LATEST, 1);
 		const FTCHARToUTF8 Utf8LobbyId(*Session->SessionInfo->GetSessionId().ToString());
 		UpdateLobbyModificationOptions.LobbyId = (EOS_LobbyId)Utf8LobbyId.Get();
 		UpdateLobbyModificationOptions.LocalUserId = EOSSubsystem->UserManager->GetLocalProductUserId(EOSSubsystem->UserManager->GetDefaultLocalUser()); // Maybe not split screen friendly
@@ -3904,7 +3971,8 @@ uint32 FOnlineSessionEOS::UpdateLobbySession(FNamedOnlineSession* Session)
 			SetLobbyAttributes(LobbyModificationHandle, Session);
 
 			EOS_Lobby_UpdateLobbyOptions UpdateLobbyOptions = { 0 };
-			UpdateLobbyOptions.ApiVersion = EOS_LOBBY_UPDATELOBBY_API_LATEST;
+			UpdateLobbyOptions.ApiVersion = 1;
+			UE_EOS_CHECK_API_MISMATCH(EOS_LOBBY_UPDATELOBBY_API_LATEST, 1);
 			UpdateLobbyOptions.LobbyModificationHandle = LobbyModificationHandle;
 
 			FName SessionName = Session->SessionName;
@@ -3982,7 +4050,8 @@ uint32 FOnlineSessionEOS::DestroyLobbySession(FNamedOnlineSession* Session, cons
 
 		// EOS will use the host migration setting to decide if the lobby is destroyed if it's the owner leaving
 		EOS_Lobby_LeaveLobbyOptions LeaveOptions = { 0 };
-		LeaveOptions.ApiVersion = EOS_LOBBY_LEAVELOBBY_API_LATEST;
+		LeaveOptions.ApiVersion = 1;
+		UE_EOS_CHECK_API_MISMATCH(EOS_LOBBY_LEAVELOBBY_API_LATEST, 1);
 		const FTCHARToUTF8 Utf8LobbyId(*SessionInfo->GetSessionId().ToString());
 		LeaveOptions.LobbyId = (EOS_LobbyId)Utf8LobbyId.Get();
 		LeaveOptions.LocalUserId = EOSSubsystem->UserManager->GetLocalProductUserId(EOSSubsystem->UserManager->GetDefaultLocalUser()); // Maybe not split screen friendly
@@ -4041,7 +4110,8 @@ uint32 FOnlineSessionEOS::FindLobbySession(int32 SearchingPlayerNum, const TShar
 	uint32 Result = ONLINE_FAIL;
 
 	EOS_Lobby_CreateLobbySearchOptions CreateLobbySearchOptions = { 0 };
-	CreateLobbySearchOptions.ApiVersion = EOS_LOBBY_CREATELOBBYSEARCH_API_LATEST;
+	CreateLobbySearchOptions.ApiVersion = 1;
+	UE_EOS_CHECK_API_MISMATCH(EOS_LOBBY_CREATELOBBYSEARCH_API_LATEST, 1);
 	CreateLobbySearchOptions.MaxResults = FMath::Clamp(SearchSettings->MaxSearchResults, 0, EOS_SESSIONS_MAX_SEARCH_RESULTS);
 
 	EOS_HLobbySearch LobbySearchHandle;
@@ -4090,7 +4160,8 @@ void FOnlineSessionEOS::StartLobbySearch(int32 SearchingPlayerNum, EOS_HLobbySea
 	SessionSearchStartInSeconds = FPlatformTime::Seconds();
 
 	EOS_LobbySearch_FindOptions FindOptions = { 0 };
-	FindOptions.ApiVersion = EOS_LOBBYSEARCH_FIND_API_LATEST;
+	FindOptions.ApiVersion = 1;
+	UE_EOS_CHECK_API_MISMATCH(EOS_LOBBYSEARCH_FIND_API_LATEST, 1);
 	FindOptions.LocalUserId = EOSSubsystem->UserManager->GetLocalProductUserId(SearchingPlayerNum);
 
 	FLobbySearchFindCallback* CallbackObj = new FLobbySearchFindCallback(FOnlineSessionEOSWeakPtr(AsShared()));
@@ -4104,14 +4175,16 @@ void FOnlineSessionEOS::StartLobbySearch(int32 SearchingPlayerNum, EOS_HLobbySea
 			CurrentSessionSearch->SearchState = EOnlineAsyncTaskState::Done;
 
 			EOS_LobbySearch_GetSearchResultCountOptions GetSearchResultCountOptions = { 0 };
-			GetSearchResultCountOptions.ApiVersion = EOS_LOBBYSEARCH_GETSEARCHRESULTCOUNT_API_LATEST;
+			GetSearchResultCountOptions.ApiVersion = 1;
+			UE_EOS_CHECK_API_MISMATCH(EOS_LOBBYSEARCH_GETSEARCHRESULTCOUNT_API_LATEST, 1);
 
 			uint32_t SearchResultsCount = EOS_LobbySearch_GetSearchResultCount(LobbySearchHandle, &GetSearchResultCountOptions);
 
 			if (SearchResultsCount > 0)
 			{
 				EOS_LobbySearch_CopySearchResultByIndexOptions CopySearchResultByIndexOptions = { 0 };
-				CopySearchResultByIndexOptions.ApiVersion = EOS_LOBBYSEARCH_COPYSEARCHRESULTBYINDEX_API_LATEST;
+				CopySearchResultByIndexOptions.ApiVersion = 1;
+				UE_EOS_CHECK_API_MISMATCH(EOS_LOBBYSEARCH_COPYSEARCHRESULTBYINDEX_API_LATEST, 1);
 
 				for (uint32_t LobbyIndex = 0; LobbyIndex < SearchResultsCount; LobbyIndex++)
 				{
@@ -4176,7 +4249,8 @@ void FOnlineSessionEOS::AddLobbySearchResult(const TSharedRef<FLobbyDetailsEOS>&
 {
 	EOS_LobbyDetails_Info* LobbyDetailsInfo = nullptr;
 	EOS_LobbyDetails_CopyInfoOptions CopyOptions = { };
-	CopyOptions.ApiVersion = EOS_LOBBYDETAILS_COPYINFO_API_LATEST;
+	CopyOptions.ApiVersion = 1;
+	UE_EOS_CHECK_API_MISMATCH(EOS_LOBBYDETAILS_COPYINFO_API_LATEST, 1);
 	EOS_EResult CopyResult = EOS_LobbyDetails_CopyInfo(LobbyDetails->LobbyDetailsHandle, &CopyOptions, &LobbyDetailsInfo);
 	if (CopyResult == EOS_EResult::EOS_Success)
 	{
@@ -4245,7 +4319,8 @@ void FOnlineSessionEOS::CopyLobbyData(const TSharedRef<FLobbyDetailsEOS>& LobbyD
 
 	// Then we copy the settings for all lobby members
 	EOS_LobbyDetails_GetMemberCountOptions CountOptions = { };
-	CountOptions.ApiVersion = EOS_LOBBYDETAILS_GETMEMBERCOUNT_API_LATEST;
+	CountOptions.ApiVersion = 1;
+	UE_EOS_CHECK_API_MISMATCH(EOS_LOBBYDETAILS_GETMEMBERCOUNT_API_LATEST, 1);
 	int32 Count = EOS_LobbyDetails_GetMemberCount(LobbyDetails->LobbyDetailsHandle, &CountOptions);
 
 	TArray<EOS_ProductUserId> TargetUserIds;
@@ -4253,7 +4328,8 @@ void FOnlineSessionEOS::CopyLobbyData(const TSharedRef<FLobbyDetailsEOS>& LobbyD
 	for (int32 Index = 0; Index < Count; Index++)
 	{
 		EOS_LobbyDetails_GetMemberByIndexOptions GetMemberByIndexOptions = { };
-		GetMemberByIndexOptions.ApiVersion = EOS_LOBBYDETAILS_GETMEMBERBYINDEX_API_LATEST;
+		GetMemberByIndexOptions.ApiVersion = 1;
+		UE_EOS_CHECK_API_MISMATCH(EOS_LOBBYDETAILS_GETMEMBERBYINDEX_API_LATEST, 1);
 		GetMemberByIndexOptions.MemberIndex = Index;
 
 		EOS_ProductUserId TargetUserId = EOS_LobbyDetails_GetMemberByIndex(LobbyDetails->LobbyDetailsHandle, &GetMemberByIndexOptions);
@@ -4271,8 +4347,8 @@ void FOnlineSessionEOS::CopyLobbyData(const TSharedRef<FLobbyDetailsEOS>& LobbyD
 				{
 					// One of the resolved ids will be the Owner's, so we'll set that too
 					EOS_LobbyDetails_GetLobbyOwnerOptions GetLobbyOwnerOptions = {};
-					GetLobbyOwnerOptions.ApiVersion = EOS_LOBBYDETAILS_GETLOBBYOWNER_API_LATEST;
-					static_assert(EOS_LOBBYDETAILS_GETLOBBYOWNER_API_LATEST == 1, "EOS_LobbyDetails_GetLobbyOwnerOptions updated, check new fields");
+					GetLobbyOwnerOptions.ApiVersion = 1;
+					UE_EOS_CHECK_API_MISMATCH(EOS_LOBBYDETAILS_GETLOBBYOWNER_API_LATEST, 1);
 
 					const EOS_ProductUserId LobbyOwner = EOS_LobbyDetails_GetLobbyOwner(LobbyDetails->LobbyDetailsHandle, &GetLobbyOwnerOptions);
 
@@ -4305,13 +4381,15 @@ void FOnlineSessionEOS::CopyLobbyAttributes(const TSharedRef<FLobbyDetailsEOS>&
 	// In this method we are updating/adding attributes, but not removing
 
 	EOS_LobbyDetails_GetAttributeCountOptions CountOptions = { };
-	CountOptions.ApiVersion = EOS_LOBBYDETAILS_GETATTRIBUTECOUNT_API_LATEST;
+	CountOptions.ApiVersion = 1;
+	UE_EOS_CHECK_API_MISMATCH(EOS_LOBBYDETAILS_GETATTRIBUTECOUNT_API_LATEST, 1);
 	int32 Count = EOS_LobbyDetails_GetAttributeCount(LobbyDetails->LobbyDetailsHandle, &CountOptions);
 
 	for (int32 Index = 0; Index < Count; Index++)
 	{
 		EOS_LobbyDetails_CopyAttributeByIndexOptions AttrOptions = { };
-		AttrOptions.ApiVersion = EOS_LOBBYDETAILS_COPYATTRIBUTEBYINDEX_API_LATEST;
+		AttrOptions.ApiVersion = 1;
+		UE_EOS_CHECK_API_MISMATCH(EOS_LOBBYDETAILS_COPYATTRIBUTEBYINDEX_API_LATEST, 1);
 		AttrOptions.AttrIndex = Index;
 
 		EOS_Lobby_Attribute* Attribute = NULL;
@@ -4384,14 +4462,16 @@ void FOnlineSessionEOS::CopyLobbyMemberAttributes(const FLobbyDetailsEOS& LobbyD
 	// In this method we are updating/adding attributes, but not removing
 
 	EOS_LobbyDetails_GetMemberAttributeCountOptions GetMemberAttributeCountOptions = {};
-	GetMemberAttributeCountOptions.ApiVersion = EOS_LOBBYDETAILS_GETMEMBERATTRIBUTECOUNT_API_LATEST;
+	GetMemberAttributeCountOptions.ApiVersion = 1;
+	UE_EOS_CHECK_API_MISMATCH(EOS_LOBBYDETAILS_GETMEMBERATTRIBUTECOUNT_API_LATEST, 1);
 	GetMemberAttributeCountOptions.TargetUserId = TargetUserId;
 
 	uint32_t MemberAttributeCount = EOS_LobbyDetails_GetMemberAttributeCount(LobbyDetails.LobbyDetailsHandle, &GetMemberAttributeCountOptions);
 	for (uint32_t MemberAttributeIndex = 0; MemberAttributeIndex < MemberAttributeCount; MemberAttributeIndex++)
 	{
 		EOS_LobbyDetails_CopyMemberAttributeByIndexOptions AttrOptions = { };
-		AttrOptions.ApiVersion = EOS_LOBBYDETAILS_COPYMEMBERATTRIBUTEBYINDEX_API_LATEST;
+		AttrOptions.ApiVersion = 1;
+		UE_EOS_CHECK_API_MISMATCH(EOS_LOBBYDETAILS_COPYMEMBERATTRIBUTEBYINDEX_API_LATEST, 1);
 		AttrOptions.AttrIndex = MemberAttributeIndex;
 
 		EOS_Lobby_Attribute* Attribute = NULL;
diff --git a/Engine/Plugins/Online/OnlineSubsystemEOS/Source/OnlineSubsystemEOS/Private/OnlineStatsEOS.cpp b/Engine/Plugins/Online/OnlineSubsystemEOS/Source/OnlineSubsystemEOS/Private/OnlineStatsEOS.cpp
--- a/Engine/Plugins/Online/OnlineSubsystemEOS/Source/OnlineSubsystemEOS/Private/OnlineStatsEOS.cpp
+++ b/Engine/Plugins/Online/OnlineSubsystemEOS/Source/OnlineSubsystemEOS/Private/OnlineStatsEOS.cpp
@@ -43,7 +43,8 @@ struct FQueryStatsOptions :
 	FQueryStatsOptions(const uint32 InNumStatIds) :
 		EOS_Stats_QueryStatsOptions()
 	{
-		ApiVersion = EOS_STATS_QUERYSTATS_API_LATEST;
+		ApiVersion = 3;
+		UE_EOS_CHECK_API_MISMATCH(EOS_STATS_QUERYSTATS_API_LATEST, 3);
 
 		StartTime = EOS_STATS_TIME_UNDEFINED;
 		EndTime = EOS_STATS_TIME_UNDEFINED;
@@ -166,7 +167,8 @@ void FOnlineStatsEOS::QueryStats(const FUniqueNetIdRef LocalUserId, const TArray
 			{
 				char StatNameANSI[EOS_OSS_STRING_BUFFER_LENGTH];
 				EOS_Stats_CopyStatByNameOptions Options = { };
-				Options.ApiVersion = EOS_STATS_COPYSTATBYNAME_API_LATEST;
+				Options.ApiVersion = 1;
+				UE_EOS_CHECK_API_MISMATCH(EOS_STATS_COPYSTATBYNAME_API_LATEST, 1);
 				Options.TargetUserId = Data->TargetUserId;
 				Options.Name = StatNameANSI;
 
@@ -290,7 +292,8 @@ void FOnlineStatsEOS::WriteStats(EOS_ProductUserId LocalUserId, EOS_ProductUserI
 	for (const TPair<FString, FOnlineStatUpdate>& Stat : PlayerStats.Stats)
 	{
 		EOS_Stats_IngestData& EOSStat = EOSData[Index];
-		EOSStat.ApiVersion = EOS_STATS_INGESTDATA_API_LATEST;
+		EOSStat.ApiVersion = 1;
+		UE_EOS_CHECK_API_MISMATCH(EOS_STATS_INGESTDATA_API_LATEST, 1);
 
 		EOSStat.IngestAmount = GetVariantValue(Stat.Value.GetValue());
 		FCStringAnsi::Strncpy(EOSStatNames[Index].StatName, TCHAR_TO_UTF8(*Stat.Key.ToUpper()), EOS_OSS_STRING_BUFFER_LENGTH);
@@ -300,7 +303,8 @@ void FOnlineStatsEOS::WriteStats(EOS_ProductUserId LocalUserId, EOS_ProductUserI
 	}
 
 	EOS_Stats_IngestStatOptions Options = { };
-	Options.ApiVersion = EOS_STATS_INGESTSTAT_API_LATEST;
+	Options.ApiVersion = 3;
+	UE_EOS_CHECK_API_MISMATCH(EOS_STATS_INGESTSTAT_API_LATEST, 3);
 	Options.LocalUserId = LocalUserId;
 	Options.TargetUserId = UserId;
 	Options.Stats = EOSData.GetData();
diff --git a/Engine/Plugins/Online/OnlineSubsystemEOS/Source/OnlineSubsystemEOS/Private/OnlineStoreEOS.cpp b/Engine/Plugins/Online/OnlineSubsystemEOS/Source/OnlineSubsystemEOS/Private/OnlineStoreEOS.cpp
--- a/Engine/Plugins/Online/OnlineSubsystemEOS/Source/OnlineSubsystemEOS/Private/OnlineStoreEOS.cpp
+++ b/Engine/Plugins/Online/OnlineSubsystemEOS/Source/OnlineSubsystemEOS/Private/OnlineStoreEOS.cpp
@@ -59,7 +59,8 @@ void FOnlineStoreEOS::QueryOffers(const FUniqueNetId& UserId, const FOnQueryOnli
 	CachedOffers.Reset();
 
 	EOS_Ecom_QueryOffersOptions Options = { };
-	Options.ApiVersion = EOS_ECOM_QUERYOFFERS_API_LATEST;
+	Options.ApiVersion = 1;
+	UE_EOS_CHECK_API_MISMATCH(EOS_ECOM_QUERYOFFERS_API_LATEST, 1);
 	Options.LocalUserId = AccountId;
 
 	FQueryOffersCallback* CallbackObj = new FQueryOffersCallback(FOnlineStoreEOSWeakPtr(AsShared()));
@@ -73,12 +74,14 @@ void FOnlineStoreEOS::QueryOffers(const FUniqueNetId& UserId, const FOnQueryOnli
 		}
 
 		EOS_Ecom_GetOfferCountOptions CountOptions = { };
-		CountOptions.ApiVersion = EOS_ECOM_GETOFFERCOUNT_API_LATEST;
+		CountOptions.ApiVersion = 1;
+		UE_EOS_CHECK_API_MISMATCH(EOS_ECOM_GETOFFERCOUNT_API_LATEST, 1);
 		CountOptions.LocalUserId = Data->LocalUserId;
 		uint32 OfferCount = EOS_Ecom_GetOfferCount(EOSSubsystem->EcomHandle, &CountOptions);
 
 		EOS_Ecom_CopyOfferByIndexOptions OfferOptions = { };
-		OfferOptions.ApiVersion = EOS_ECOM_COPYOFFERBYINDEX_API_LATEST;
+		OfferOptions.ApiVersion = 3;
+		UE_EOS_CHECK_API_MISMATCH(EOS_ECOM_COPYOFFERBYINDEX_API_LATEST, 3);
 		OfferOptions.LocalUserId = Data->LocalUserId;
 		// Iterate and parse the offer list
 		for (uint32 OfferIndex = 0; OfferIndex < OfferCount; OfferIndex++)
@@ -176,13 +179,15 @@ void FOnlineStoreEOS::Checkout(const FUniqueNetId& UserId, const FPurchaseChecko
 	// Convert the items
 	for (int32 Index = 0; Index < NumItems; Index++)
 	{
-		Entries[Index].ApiVersion = EOS_ECOM_CHECKOUTENTRY_API_LATEST;
+		Entries[Index].ApiVersion = 1;
+		UE_EOS_CHECK_API_MISMATCH(EOS_ECOM_CHECKOUTENTRY_API_LATEST, 1);
 		Entries[Index].OfferId = Items[Index];
 		FCStringAnsi::Strncpy(Items[Index], TCHAR_TO_UTF8(*CheckoutRequest.PurchaseOffers[Index].OfferId), EOS_ECOM_TRANSACTIONID_MAXIMUM_LENGTH);
 	}
 
 	EOS_Ecom_CheckoutOptions Options = { };
-	Options.ApiVersion = EOS_ECOM_CHECKOUT_API_LATEST;
+	Options.ApiVersion = 1;
+	UE_EOS_CHECK_API_MISMATCH(EOS_ECOM_CHECKOUT_API_LATEST, 1);
 	Options.LocalUserId = AccountId;
 	Options.EntryCount = NumItems;
 	Options.Entries = (const EOS_Ecom_CheckoutEntry*)Entries.GetData();
@@ -266,7 +271,8 @@ void FOnlineStoreEOS::QueryReceipts(const FUniqueNetId& UserId, bool bRestoreRec
 	CachedReceipts.Reset();
 
 	EOS_Ecom_QueryEntitlementsOptions Options = { };
-	Options.ApiVersion = EOS_ECOM_QUERYENTITLEMENTS_API_LATEST;
+	Options.ApiVersion = 2;
+	UE_EOS_CHECK_API_MISMATCH(EOS_ECOM_QUERYENTITLEMENTS_API_LATEST, 2);
 	Options.LocalUserId = AccountId;
 	Options.bIncludeRedeemed = bRestoreReceipts ? EOS_TRUE : EOS_FALSE;
 
@@ -282,13 +288,15 @@ void FOnlineStoreEOS::QueryReceipts(const FUniqueNetId& UserId, bool bRestoreRec
 		}
 
 		EOS_Ecom_GetEntitlementsCountOptions CountOptions = { };
-		CountOptions.ApiVersion = EOS_ECOM_GETENTITLEMENTSCOUNT_API_LATEST;
+		CountOptions.ApiVersion = 1;
+		UE_EOS_CHECK_API_MISMATCH(EOS_ECOM_GETENTITLEMENTSCOUNT_API_LATEST, 1);
 		CountOptions.LocalUserId = Data->LocalUserId;
 		uint32 Count = EOS_Ecom_GetEntitlementsCount(EOSSubsystem->EcomHandle, &CountOptions);
 		CachedReceipts.AddZeroed(Count);
 
 		EOS_Ecom_CopyEntitlementByIndexOptions CopyOptions = { };
-		CopyOptions.ApiVersion = EOS_ECOM_COPYENTITLEMENTBYINDEX_API_LATEST;
+		CopyOptions.ApiVersion = 1;
+		UE_EOS_CHECK_API_MISMATCH(EOS_ECOM_COPYENTITLEMENTBYINDEX_API_LATEST, 1);
 		CopyOptions.LocalUserId = Data->LocalUserId;
 
 		for (uint32 Index = 0; Index < Count; Index++)
@@ -348,7 +356,8 @@ void FOnlineStoreEOS::FinalizeReceiptValidationInfo(const FUniqueNetId& UserId,
 	Ids[0] = EntitlementId.Get();
 
 	EOS_Ecom_RedeemEntitlementsOptions Options = { };
-	Options.ApiVersion = EOS_ECOM_REDEEMENTITLEMENTS_API_LATEST;
+	Options.ApiVersion = 2;
+	UE_EOS_CHECK_API_MISMATCH(EOS_ECOM_REDEEMENTITLEMENTS_API_LATEST, 2);
 	Options.LocalUserId = AccountId;
 	Options.EntitlementIdCount = 1;
 	Options.EntitlementIds = Ids;
diff --git a/Engine/Plugins/Online/OnlineSubsystemEOS/Source/OnlineSubsystemEOS/Private/OnlineSubsystemEOS.cpp b/Engine/Plugins/Online/OnlineSubsystemEOS/Source/OnlineSubsystemEOS/Private/OnlineSubsystemEOS.cpp
--- a/Engine/Plugins/Online/OnlineSubsystemEOS/Source/OnlineSubsystemEOS/Private/OnlineSubsystemEOS.cpp
+++ b/Engine/Plugins/Online/OnlineSubsystemEOS/Source/OnlineSubsystemEOS/Private/OnlineSubsystemEOS.cpp
@@ -131,7 +131,8 @@ struct FEOSPlatformOptions :
 	FEOSPlatformOptions() :
 		EOS_Platform_Options()
 	{
-		ApiVersion = EOS_PLATFORM_OPTIONS_API_LATEST;
+		ApiVersion = 12;
+		UE_EOS_CHECK_API_MISMATCH(EOS_PLATFORM_OPTIONS_API_LATEST, 12);
 		ProductId = ProductIdAnsi;
 		SandboxId = SandboxIdAnsi;
 		DeploymentId = DeploymentIdAnsi;
@@ -249,7 +250,8 @@ bool FOnlineSubsystemEOS::PlatformCreate()
 
 #if WITH_EOS_RTC
 	EOS_Platform_RTCOptions RtcOptions = { 0 };
-	RtcOptions.ApiVersion = EOS_PLATFORM_RTCOPTIONS_API_LATEST;
+	RtcOptions.ApiVersion = 1;
+	UE_EOS_CHECK_API_MISMATCH(EOS_PLATFORM_RTCOPTIONS_API_LATEST, 1);
 	RtcOptions.PlatformSpecificOptions = nullptr;
 	PlatformOptions.RTCOptions = &RtcOptions;
 #endif
diff --git a/Engine/Plugins/Online/OnlineSubsystemEOS/Source/OnlineSubsystemEOS/Private/OnlineTitleFileEOS.cpp b/Engine/Plugins/Online/OnlineSubsystemEOS/Source/OnlineSubsystemEOS/Private/OnlineTitleFileEOS.cpp
--- a/Engine/Plugins/Online/OnlineSubsystemEOS/Source/OnlineSubsystemEOS/Private/OnlineTitleFileEOS.cpp
+++ b/Engine/Plugins/Online/OnlineSubsystemEOS/Source/OnlineSubsystemEOS/Private/OnlineTitleFileEOS.cpp
@@ -95,7 +95,8 @@ void FOnlineTitleFileEOS::DeleteCachedFiles(bool bSkipEnumerated)
 	}
 
 	EOS_TitleStorage_DeleteCacheOptions DeleteCacheOptions = { };
-	DeleteCacheOptions.ApiVersion = EOS_TITLESTORAGE_DELETECACHEOPTIONS_API_LATEST;
+	DeleteCacheOptions.ApiVersion = 1;
+	UE_EOS_CHECK_API_MISMATCH(EOS_TITLESTORAGE_DELETECACHEOPTIONS_API_LATEST, 1);
 	DeleteCacheOptions.LocalUserId = EOSSubsystem->UserManager->GetLocalProductUserId();	// Get a local user if one is available, but this is not required
 
 	FDeleteCacheCompleteCallback* CallbackObj = new FDeleteCacheCompleteCallback(FOnlineTitleFileEOSWeakPtr(AsShared()));
@@ -164,7 +165,8 @@ bool FOnlineTitleFileEOS::EnumerateFiles(const FPagedQuery& Page)
 
 
 	EOS_TitleStorage_QueryFileListOptions QueryFileListOptions = { };
-	QueryFileListOptions.ApiVersion = EOS_TITLESTORAGE_QUERYFILELISTOPTIONS_API_LATEST;
+	QueryFileListOptions.ApiVersion = 1;
+	UE_EOS_CHECK_API_MISMATCH(EOS_TITLESTORAGE_QUERYFILELISTOPTIONS_API_LATEST, 1);
 	QueryFileListOptions.LocalUserId = EOSSubsystem->UserManager->GetLocalProductUserId();	// Get a local user if one is available, but this is not required
 	QueryFileListOptions.ListOfTags = const_cast<const char**>(AnsiTags.GetData());
 	QueryFileListOptions.ListOfTagsCount = AnsiTags.Num();
@@ -184,7 +186,8 @@ bool FOnlineTitleFileEOS::EnumerateFiles(const FPagedQuery& Page)
 			for (uint32 Index = 0; Index < FileCount; ++Index)
 			{
 				EOS_TitleStorage_CopyFileMetadataAtIndexOptions CopyFileMetadataAtIndexOptions = { };
-				CopyFileMetadataAtIndexOptions.ApiVersion = EOS_TITLESTORAGE_COPYFILEMETADATAATINDEXOPTIONS_API_LATEST;
+				CopyFileMetadataAtIndexOptions.ApiVersion = 1;
+				UE_EOS_CHECK_API_MISMATCH(EOS_TITLESTORAGE_COPYFILEMETADATAATINDEXOPTIONS_API_LATEST, 1);
 				CopyFileMetadataAtIndexOptions.LocalUserId = Data->LocalUserId;
 				CopyFileMetadataAtIndexOptions.Index = Index;
 
@@ -339,7 +342,8 @@ bool FOnlineTitleFileEOS::ReadFile(const FString& FileName)
 	}
 
 	EOS_TitleStorage_ReadFileOptions ReadFileOptions = { };
-	ReadFileOptions.ApiVersion = EOS_TITLESTORAGE_READFILEOPTIONS_API_LATEST;
+	ReadFileOptions.ApiVersion = 1;
+	UE_EOS_CHECK_API_MISMATCH(EOS_TITLESTORAGE_READFILEOPTIONS_API_LATEST, 1);
 	ReadFileOptions.LocalUserId = EOSSubsystem->UserManager->GetLocalProductUserId();	// Get a local user if one is available, but this is not required
 	ReadFileOptions.Filename = AnsiFileName;
 	ReadFileOptions.ReadChunkLengthBytes = (uint32)ReadChunkSize;
diff --git a/Engine/Plugins/Online/OnlineSubsystemEOS/Source/OnlineSubsystemEOS/Private/OnlineUserCloudEOS.cpp b/Engine/Plugins/Online/OnlineSubsystemEOS/Source/OnlineSubsystemEOS/Private/OnlineUserCloudEOS.cpp
--- a/Engine/Plugins/Online/OnlineSubsystemEOS/Source/OnlineSubsystemEOS/Private/OnlineUserCloudEOS.cpp
+++ b/Engine/Plugins/Online/OnlineSubsystemEOS/Source/OnlineSubsystemEOS/Private/OnlineUserCloudEOS.cpp
@@ -191,7 +191,8 @@ void FOnlineUserCloudEOS::EnumerateUserFiles(const FUniqueNetId& UserId)
 	}
 
 	EOS_PlayerDataStorage_QueryFileListOptions Options = {};
-	Options.ApiVersion = EOS_PLAYERDATASTORAGE_QUERYFILELISTOPTIONS_API_LATEST;
+	Options.ApiVersion = 1;
+	UE_EOS_CHECK_API_MISMATCH(EOS_PLAYERDATASTORAGE_QUERYFILELISTOPTIONS_API_LATEST, 1);
 	Options.LocalUserId = LocalUserId;
 
 	FOnQueryFileListCallback* CallbackObj = new FOnQueryFileListCallback(FOnlineUserCloudEOSWeakPtr(AsShared()));
@@ -213,7 +214,8 @@ void FOnlineUserCloudEOS::EnumerateUserFiles(const FUniqueNetId& UserId)
 				for (uint32 Index = 0; Index < FileCount; ++Index)
 				{
 					EOS_PlayerDataStorage_CopyFileMetadataAtIndexOptions CopyFileMetadataAtIndexOptions = { };
-					CopyFileMetadataAtIndexOptions.ApiVersion = EOS_PLAYERDATASTORAGE_COPYFILEMETADATAATINDEXOPTIONS_API_LATEST;
+					CopyFileMetadataAtIndexOptions.ApiVersion = 1;
+					UE_EOS_CHECK_API_MISMATCH(EOS_PLAYERDATASTORAGE_COPYFILEMETADATAATINDEXOPTIONS_API_LATEST, 1);
 					CopyFileMetadataAtIndexOptions.LocalUserId = Data->LocalUserId;
 					CopyFileMetadataAtIndexOptions.Index = Index;
 
@@ -423,7 +425,8 @@ bool FOnlineUserCloudEOS::ReadUserFile(const FUniqueNetId& UserId, const FString
 	const FUniqueNetIdEOS& UserEOSId = FUniqueNetIdEOS::Cast(UserId);
 
 	EOS_PlayerDataStorage_ReadFileOptions ReadFileOptions = {};
-	ReadFileOptions.ApiVersion = EOS_PLAYERDATASTORAGE_READFILEOPTIONS_API_LATEST;
+	ReadFileOptions.ApiVersion = 1;
+	UE_EOS_CHECK_API_MISMATCH(EOS_PLAYERDATASTORAGE_READFILEOPTIONS_API_LATEST, 1);
 	ReadFileOptions.LocalUserId = UserEOSId.GetProductUserId();
 	ReadFileOptions.Filename = FileNameUtf8.Get();
 	ReadFileOptions.ReadChunkLengthBytes = (uint32_t)ReadChunkSize;
@@ -601,7 +604,8 @@ bool FOnlineUserCloudEOS::WriteUserFile(const FUniqueNetId& UserId, const FStrin
 	const FUniqueNetIdEOS& UserEOSId = FUniqueNetIdEOS::Cast(UserId);
 
 	EOS_PlayerDataStorage_WriteFileOptions WriteFileOptions = {};
-	WriteFileOptions.ApiVersion = EOS_PLAYERDATASTORAGE_WRITEFILEOPTIONS_API_LATEST;
+	WriteFileOptions.ApiVersion = 1;
+	UE_EOS_CHECK_API_MISMATCH(EOS_PLAYERDATASTORAGE_WRITEFILEOPTIONS_API_LATEST, 1);
 	WriteFileOptions.LocalUserId = UserEOSId.GetProductUserId();
 	WriteFileOptions.Filename = FileNameUtf8.Get();
 	WriteFileOptions.ChunkLengthBytes = (uint32_t)ReadChunkSize;
@@ -755,7 +759,8 @@ bool FOnlineUserCloudEOS::DeleteUserFile(const FUniqueNetId& UserId, const FStri
 		const FUniqueNetIdEOS& UserEOSId = FUniqueNetIdEOS::Cast(UserId);
 
 		EOS_PlayerDataStorage_DeleteFileOptions Options = {};
-		Options.ApiVersion = EOS_PLAYERDATASTORAGE_DELETEFILEOPTIONS_API_LATEST;
+		Options.ApiVersion = 1;
+		UE_EOS_CHECK_API_MISMATCH(EOS_PLAYERDATASTORAGE_DELETEFILEOPTIONS_API_LATEST, 1);
 		Options.LocalUserId = UserEOSId.GetProductUserId();
 		Options.Filename = FileNameUtf8.Get();
 
diff --git a/Engine/Plugins/Online/OnlineSubsystemEOS/Source/OnlineSubsystemEOS/Private/UserManagerEOS.cpp b/Engine/Plugins/Online/OnlineSubsystemEOS/Source/OnlineSubsystemEOS/Private/UserManagerEOS.cpp
--- a/Engine/Plugins/Online/OnlineSubsystemEOS/Source/OnlineSubsystemEOS/Private/UserManagerEOS.cpp
+++ b/Engine/Plugins/Online/OnlineSubsystemEOS/Source/OnlineSubsystemEOS/Private/UserManagerEOS.cpp
@@ -159,7 +159,8 @@ void FUserManagerEOS::Init()
 	{
 		// Adding subscription to external ui display change event
 		EOS_UI_AddNotifyDisplaySettingsUpdatedOptions Options = {};
-		Options.ApiVersion = EOS_UI_ADDNOTIFYDISPLAYSETTINGSUPDATED_API_LATEST;
+		Options.ApiVersion = 1;
+		UE_EOS_CHECK_API_MISMATCH(EOS_UI_ADDNOTIFYDISPLAYSETTINGSUPDATED_API_LATEST, 1);
 
 		FOnDisplaySettingsUpdatedCallback* CallbackObj = new FOnDisplaySettingsUpdatedCallback(AsWeak());
 		DisplaySettingsUpdatedCallback = CallbackObj;
@@ -324,7 +325,8 @@ struct FAuthCredentials :
 	FAuthCredentials() :
 		EOS_Auth_Credentials()
 	{
-		ApiVersion = EOS_AUTH_CREDENTIALS_API_LATEST;
+		ApiVersion = 3;
+		UE_EOS_CHECK_API_MISMATCH(EOS_AUTH_CREDENTIALS_API_LATEST, 3);
 		Id = IdAnsi;
 		Token = TokenAnsi;
 
@@ -377,7 +379,8 @@ struct FAuthCredentials :
 
 	void Init(EOS_EExternalCredentialType InExternalType, const FString& InTokenString)
 	{
-		ApiVersion = EOS_AUTH_CREDENTIALS_API_LATEST;
+		ApiVersion = 3;
+		UE_EOS_CHECK_API_MISMATCH(EOS_AUTH_CREDENTIALS_API_LATEST, 3);
 		Type = EOS_ELoginCredentialType::EOS_LCT_ExternalAuth;
 		ExternalType = InExternalType;
 		Id = IdAnsi;
@@ -388,7 +391,8 @@ struct FAuthCredentials :
 
 	void Init(EOS_EExternalCredentialType InExternalType, const TArray<uint8>& InToken)
 	{
-		ApiVersion = EOS_AUTH_CREDENTIALS_API_LATEST;
+		ApiVersion = 3;
+		UE_EOS_CHECK_API_MISMATCH(EOS_AUTH_CREDENTIALS_API_LATEST, 3);
 		Type = EOS_ELoginCredentialType::EOS_LCT_ExternalAuth;
 		ExternalType = InExternalType;
 		Id = IdAnsi;
@@ -439,7 +443,8 @@ bool FUserManagerEOS::Login(int32 LocalUserNum, const FOnlineAccountCredentials&
 	}
 
 	EOS_Auth_LoginOptions LoginOptions = { };
-	LoginOptions.ApiVersion = EOS_AUTH_LOGIN_API_LATEST;
+	LoginOptions.ApiVersion = 2;
+	UE_EOS_CHECK_API_MISMATCH(EOS_AUTH_LOGIN_API_LATEST, 2);
 
 	LoginOptions.ScopeFlags = EOS_EAuthScopeFlags::EOS_AS_NoFlags;
 	for (const FString& FlagsStr : Settings.AuthScopeFlags)
@@ -537,7 +542,8 @@ bool FUserManagerEOS::Login(int32 LocalUserNum, const FOnlineAccountCredentials&
 				};
 
 				EOS_Auth_DeletePersistentAuthOptions DeletePersistentAuthOptions;
-				DeletePersistentAuthOptions.ApiVersion = EOS_AUTH_DELETEPERSISTENTAUTH_API_LATEST;
+				DeletePersistentAuthOptions.ApiVersion = 2;
+				UE_EOS_CHECK_API_MISMATCH(EOS_AUTH_DELETEPERSISTENTAUTH_API_LATEST, 2);
 				DeletePersistentAuthOptions.RefreshToken = nullptr;
 				EOS_Auth_DeletePersistentAuth(EOSSubsystem->AuthHandle, &DeletePersistentAuthOptions, (void*)DeleteAuthCallbackObj, DeleteAuthCallbackObj->GetCallbackPtr());
 			}
@@ -567,7 +573,8 @@ void FUserManagerEOS::LoginViaExternalAuth(int32 LocalUserNum)
 				}
 
 				EOS_Auth_LoginOptions LoginOptions = { };
-				LoginOptions.ApiVersion = EOS_AUTH_LOGIN_API_LATEST;
+				LoginOptions.ApiVersion = 2;
+				UE_EOS_CHECK_API_MISMATCH(EOS_AUTH_LOGIN_API_LATEST, 2);
 				LoginOptions.ScopeFlags = EOS_EAuthScopeFlags::EOS_AS_BasicProfile | EOS_EAuthScopeFlags::EOS_AS_FriendsList | EOS_EAuthScopeFlags::EOS_AS_Presence;
 
 				check(LocalUserNumToLastLoginCredentials.Contains(LocalUserNum));
@@ -605,7 +612,8 @@ struct FLinkAccountOptions :
 	FLinkAccountOptions(EOS_ContinuanceToken Token)
 		: EOS_Auth_LinkAccountOptions()
 	{
-		ApiVersion = EOS_AUTH_LINKACCOUNT_API_LATEST;
+		ApiVersion = 1;
+		UE_EOS_CHECK_API_MISMATCH(EOS_AUTH_LINKACCOUNT_API_LATEST, 1);
 		ContinuanceToken = Token;
 	}
 };
@@ -655,7 +663,8 @@ struct FConnectCredentials :
 
 	void Init(EOS_EExternalCredentialType InType, const FString& InTokenString)
 	{
-		ApiVersion = EOS_CONNECT_CREDENTIALS_API_LATEST;
+		ApiVersion = 1;
+		UE_EOS_CHECK_API_MISMATCH(EOS_CONNECT_CREDENTIALS_API_LATEST, 1);
 		Token = TokenAnsi;
 		Type = InType;
 
@@ -664,7 +673,8 @@ struct FConnectCredentials :
 
 	void Init(EOS_EExternalCredentialType InType, const TArray<uint8>& InToken)
 	{
-		ApiVersion = EOS_CONNECT_CREDENTIALS_API_LATEST;
+		ApiVersion = 1;
+		UE_EOS_CHECK_API_MISMATCH(EOS_CONNECT_CREDENTIALS_API_LATEST, 1);
 		Token = TokenAnsi;
 		Type = InType;
 
@@ -694,12 +704,14 @@ bool FUserManagerEOS::ConnectLoginNoEAS(int32 LocalUserNum)
 				check(LocalUserNumToLastLoginCredentials.Contains(LocalUserNum));
 				FConnectCredentials Credentials(ToEOS_EExternalCredentialType(GetPlatformOSS()->GetSubsystemName(), *LocalUserNumToLastLoginCredentials[LocalUserNum]), AuthToken);
 				EOS_Connect_LoginOptions Options = { };
-				Options.ApiVersion = EOS_CONNECT_LOGIN_API_LATEST;
+				Options.ApiVersion = 2;
+				UE_EOS_CHECK_API_MISMATCH(EOS_CONNECT_LOGIN_API_LATEST, 2);
 				Options.Credentials = &Credentials;
 
 #if ADD_USER_LOGIN_INFO
 				EOS_Connect_UserLoginInfo UserLoginInfo = {};
-				UserLoginInfo.ApiVersion = EOS_CONNECT_USERLOGININFO_API_LATEST;
+				UserLoginInfo.ApiVersion = 1;
+				UE_EOS_CHECK_API_MISMATCH(EOS_CONNECT_USERLOGININFO_API_LATEST, 1);
 				const FTCHARToUTF8 DisplayNameUtf8(*GetPlatformDisplayName(LocalUserNum));
 				UserLoginInfo.DisplayName = DisplayNameUtf8.Get();
 
@@ -737,18 +749,21 @@ bool FUserManagerEOS::ConnectLoginEAS(int32 LocalUserNum, EOS_EpicAccountId Acco
 {
 	EOS_Auth_Token* AuthToken = nullptr;
 	EOS_Auth_CopyUserAuthTokenOptions CopyOptions = { };
-	CopyOptions.ApiVersion = EOS_AUTH_COPYUSERAUTHTOKEN_API_LATEST;
+	CopyOptions.ApiVersion = 1;
+	UE_EOS_CHECK_API_MISMATCH(EOS_AUTH_COPYUSERAUTHTOKEN_API_LATEST, 1);
 
 	EOS_EResult CopyResult = EOS_Auth_CopyUserAuthToken(EOSSubsystem->AuthHandle, &CopyOptions, AccountId, &AuthToken);
 	if (CopyResult == EOS_EResult::EOS_Success)
 	{
 		EOS_Connect_Credentials Credentials = { };
-		Credentials.ApiVersion = EOS_CONNECT_CREDENTIALS_API_LATEST;
+		Credentials.ApiVersion = 1;
+		UE_EOS_CHECK_API_MISMATCH(EOS_CONNECT_CREDENTIALS_API_LATEST, 1);
 		Credentials.Type = EOS_EExternalCredentialType::EOS_ECT_EPIC;
 		Credentials.Token = AuthToken->AccessToken;
 
 		EOS_Connect_LoginOptions Options = { };
-		Options.ApiVersion = EOS_CONNECT_LOGIN_API_LATEST;
+		Options.ApiVersion = 2;
+		UE_EOS_CHECK_API_MISMATCH(EOS_CONNECT_LOGIN_API_LATEST, 2);
 		Options.Credentials = &Credentials;
 
 		FConnectLoginCallback* CallbackObj = new FConnectLoginCallback(AsWeak());
@@ -796,7 +811,8 @@ void FUserManagerEOS::RefreshConnectLogin(int32 LocalUserNum)
 		EOS_EpicAccountId AccountId = UserNumToAccountIdMap[LocalUserNum];
 		EOS_Auth_Token* AuthToken = nullptr;
 		EOS_Auth_CopyUserAuthTokenOptions CopyOptions = { };
-		CopyOptions.ApiVersion = EOS_AUTH_COPYUSERAUTHTOKEN_API_LATEST;
+		CopyOptions.ApiVersion = 1;
+		UE_EOS_CHECK_API_MISMATCH(EOS_AUTH_COPYUSERAUTHTOKEN_API_LATEST, 1);
 
 		EOS_EResult CopyResult = EOS_Auth_CopyUserAuthToken(EOSSubsystem->AuthHandle, &CopyOptions, AccountId, &AuthToken);
 		if (CopyResult == EOS_EResult::EOS_Success)
@@ -808,12 +824,14 @@ void FUserManagerEOS::RefreshConnectLogin(int32 LocalUserNum)
 			UpdateUserInfo(UserAccountRef, AccountId, AccountId);
 
 			EOS_Connect_Credentials Credentials = { };
-			Credentials.ApiVersion = EOS_CONNECT_CREDENTIALS_API_LATEST;
+			Credentials.ApiVersion = 1;
+			UE_EOS_CHECK_API_MISMATCH(EOS_CONNECT_CREDENTIALS_API_LATEST, 1);
 			Credentials.Type = EOS_EExternalCredentialType::EOS_ECT_EPIC;
 			Credentials.Token = AuthToken->AccessToken;
 
 			EOS_Connect_LoginOptions Options = { };
-			Options.ApiVersion = EOS_CONNECT_LOGIN_API_LATEST;
+			Options.ApiVersion = 2;
+			UE_EOS_CHECK_API_MISMATCH(EOS_CONNECT_LOGIN_API_LATEST, 2);
 			Options.Credentials = &Credentials;
 
 			FConnectLoginCallback* CallbackObj = new FConnectLoginCallback(AsWeak());
@@ -856,7 +874,8 @@ void FUserManagerEOS::RefreshConnectLogin(int32 LocalUserNum)
 					EOS_EExternalCredentialType CredType = ToEOS_EExternalCredentialType(GetPlatformOSS()->GetSubsystemName(), Creds);
 					FConnectCredentials Credentials(CredType, AuthToken);
 					EOS_Connect_LoginOptions Options = { };
-					Options.ApiVersion = EOS_CONNECT_LOGIN_API_LATEST;
+					Options.ApiVersion = 2;
+					UE_EOS_CHECK_API_MISMATCH(EOS_CONNECT_LOGIN_API_LATEST, 2);
 					Options.Credentials = &Credentials;
 
 					FConnectLoginCallback* CallbackObj = new FConnectLoginCallback(AsWeak());
@@ -879,7 +898,8 @@ typedef TEOSCallback<EOS_Connect_OnCreateUserCallback, EOS_Connect_CreateUserCal
 void FUserManagerEOS::CreateConnectedLogin(int32 LocalUserNum, EOS_EpicAccountId AccountId, EOS_ContinuanceToken Token)
 {
 	EOS_Connect_CreateUserOptions Options = { };
-	Options.ApiVersion = EOS_CONNECT_CREATEUSER_API_LATEST;
+	Options.ApiVersion = 1;
+	UE_EOS_CHECK_API_MISMATCH(EOS_CONNECT_CREATEUSER_API_LATEST, 1);
 	Options.ContinuanceToken = Token;
 
 	FCreateUserCallback* CallbackObj = new FCreateUserCallback(AsWeak());
@@ -917,7 +937,8 @@ void FUserManagerEOS::FullLoginCallback(int32 LocalUserNum, EOS_EpicAccountId Ac
 		};
 
 		EOS_Auth_AddNotifyLoginStatusChangedOptions Options = { };
-		Options.ApiVersion = EOS_AUTH_ADDNOTIFYLOGINSTATUSCHANGED_API_LATEST;
+		Options.ApiVersion = 1;
+		UE_EOS_CHECK_API_MISMATCH(EOS_AUTH_ADDNOTIFYLOGINSTATUSCHANGED_API_LATEST, 1);
 		LoginNotificationId = EOS_Auth_AddNotifyLoginStatusChanged(EOSSubsystem->AuthHandle, &Options, CallbackObj, CallbackObj->GetCallbackPtr());
 	}
 	// Register for friends updates if not set yet
@@ -931,7 +952,8 @@ void FUserManagerEOS::FullLoginCallback(int32 LocalUserNum, EOS_EpicAccountId Ac
 		};
 
 		EOS_Friends_AddNotifyFriendsUpdateOptions Options = { };
-		Options.ApiVersion = EOS_FRIENDS_ADDNOTIFYFRIENDSUPDATE_API_LATEST;
+		Options.ApiVersion = 1;
+		UE_EOS_CHECK_API_MISMATCH(EOS_FRIENDS_ADDNOTIFYFRIENDSUPDATE_API_LATEST, 1);
 		FriendsNotificationId = EOS_Friends_AddNotifyFriendsUpdate(EOSSubsystem->FriendsHandle, &Options, CallbackObj, CallbackObj->GetCallbackPtr());
 	}
 	// Register for presence updates if not set yet
@@ -950,7 +972,8 @@ void FUserManagerEOS::FullLoginCallback(int32 LocalUserNum, EOS_EpicAccountId Ac
 		};
 
 		EOS_Presence_AddNotifyOnPresenceChangedOptions Options = { };
-		Options.ApiVersion = EOS_PRESENCE_ADDNOTIFYONPRESENCECHANGED_API_LATEST;
+		Options.ApiVersion = 1;
+		UE_EOS_CHECK_API_MISMATCH(EOS_PRESENCE_ADDNOTIFYONPRESENCECHANGED_API_LATEST, 1);
 		PresenceNotificationId = EOS_Presence_AddNotifyOnPresenceChanged(EOSSubsystem->PresenceHandle, &Options, CallbackObj, CallbackObj->GetCallbackPtr());
 	}
 	// Add auth refresh notification if not set for this user yet
@@ -967,7 +990,8 @@ void FUserManagerEOS::FullLoginCallback(int32 LocalUserNum, EOS_EpicAccountId Ac
 		};
 
 		EOS_Connect_AddNotifyAuthExpirationOptions Options = { };
-		Options.ApiVersion = EOS_CONNECT_ADDNOTIFYAUTHEXPIRATION_API_LATEST;
+		Options.ApiVersion = 1;
+		UE_EOS_CHECK_API_MISMATCH(EOS_CONNECT_ADDNOTIFYAUTHEXPIRATION_API_LATEST, 1);
 		NotificationPair->NotificationId = EOS_Connect_AddNotifyAuthExpiration(EOSSubsystem->ConnectHandle, &Options, CallbackObj, CallbackObj->GetCallbackPtr());
 	}
 
@@ -1011,13 +1035,15 @@ bool FUserManagerEOS::Logout(int32 LocalUserNum)
 		};
 
 		EOS_Auth_DeletePersistentAuthOptions DeletePersistentAuthOptions;
-		DeletePersistentAuthOptions.ApiVersion = EOS_AUTH_DELETEPERSISTENTAUTH_API_LATEST;
+		DeletePersistentAuthOptions.ApiVersion = 2;
+		UE_EOS_CHECK_API_MISMATCH(EOS_AUTH_DELETEPERSISTENTAUTH_API_LATEST, 2);
 		DeletePersistentAuthOptions.RefreshToken = nullptr;
 		EOS_Auth_DeletePersistentAuth(EOSSubsystem->AuthHandle, &DeletePersistentAuthOptions, (void*)DeleteAuthCallbackObj, DeleteAuthCallbackObj->GetCallbackPtr());
 	};
 
 	EOS_Auth_LogoutOptions LogoutOptions = { };
-	LogoutOptions.ApiVersion = EOS_AUTH_LOGOUT_API_LATEST;
+	LogoutOptions.ApiVersion = 1;
+	UE_EOS_CHECK_API_MISMATCH(EOS_AUTH_LOGOUT_API_LATEST, 1);
 	LogoutOptions.LocalUserId = UserId->GetEpicAccountId();
 
 	EOS_Auth_Logout(EOSSubsystem->AuthHandle, &LogoutOptions, CallbackObj, CallbackObj->GetCallbackPtr());
@@ -1083,7 +1109,8 @@ void FUserManagerEOS::AddLocalUser(int32 LocalUserNum, EOS_EpicAccountId EpicAcc
 	// Get auth token info
 	EOS_Auth_Token* AuthToken = nullptr;
 	EOS_Auth_CopyUserAuthTokenOptions Options = { };
-	Options.ApiVersion = EOS_AUTH_COPYUSERAUTHTOKEN_API_LATEST;
+	Options.ApiVersion = 1;
+	UE_EOS_CHECK_API_MISMATCH(EOS_AUTH_COPYUSERAUTHTOKEN_API_LATEST, 1);
 
 	EOS_EResult CopyResult = EOS_Auth_CopyUserAuthToken(EOSSubsystem->AuthHandle, &Options, EpicAccountId, &AuthToken);
 	if (CopyResult == EOS_EResult::EOS_Success)
@@ -1098,7 +1125,8 @@ void FUserManagerEOS::AddLocalUser(int32 LocalUserNum, EOS_EpicAccountId EpicAcc
 void FUserManagerEOS::UpdateUserInfo(IAttributeAccessInterfaceRef AttributeAccessRef, EOS_EpicAccountId LocalId, EOS_EpicAccountId AccountId)
 {
 	EOS_UserInfo_CopyUserInfoOptions Options = { };
-	Options.ApiVersion = EOS_USERINFO_COPYUSERINFO_API_LATEST;
+	Options.ApiVersion = 3;
+	UE_EOS_CHECK_API_MISMATCH(EOS_USERINFO_COPYUSERINFO_API_LATEST, 3);
 	Options.LocalUserId = LocalId;
 	Options.TargetUserId = AccountId;
 
@@ -1284,7 +1312,8 @@ bool FUserManagerEOS::GetEpicAccountIdFromProductUserId(const EOS_ProductUserId&
 	int32 EpicIdStrSize = sizeof(EpicIdStr);
 
 	EOS_Connect_GetProductUserIdMappingOptions Options = { };
-	Options.ApiVersion = EOS_CONNECT_GETPRODUCTUSERIDMAPPING_API_LATEST;
+	Options.ApiVersion = 1;
+	UE_EOS_CHECK_API_MISMATCH(EOS_CONNECT_GETPRODUCTUSERIDMAPPING_API_LATEST, 1);
 	Options.AccountIdType = EOS_EExternalAccountType::EOS_EAT_EPIC;
 	Options.LocalUserId = GetLocalProductUserId();
 	Options.TargetProductUserId = ProductUserId;
@@ -1340,7 +1369,8 @@ void FUserManagerEOS::ResolveUniqueNetIds(const TArray<EOS_ProductUserId>& Produ
 	if (!ProductUserIdsToResolve.IsEmpty())
 	{
 		EOS_Connect_QueryProductUserIdMappingsOptions QueryProductUserIdMappingsOptions = {};
-		QueryProductUserIdMappingsOptions.ApiVersion = EOS_CONNECT_QUERYPRODUCTUSERIDMAPPINGS_API_LATEST;
+		QueryProductUserIdMappingsOptions.ApiVersion = 2;
+		UE_EOS_CHECK_API_MISMATCH(EOS_CONNECT_QUERYPRODUCTUSERIDMAPPINGS_API_LATEST, 2);
 		QueryProductUserIdMappingsOptions.LocalUserId = EOSSubsystem->UserManager->GetLocalProductUserId(0);
 		QueryProductUserIdMappingsOptions.ProductUserIds = ProductUserIdsToResolve.GetData();
 		QueryProductUserIdMappingsOptions.ProductUserIdCount = ProductUserIdsToResolve.Num();
@@ -1602,7 +1632,8 @@ typedef TEOSCallback<EOS_UI_OnShowFriendsCallback, EOS_UI_ShowFriendsCallbackInf
 bool FUserManagerEOS::ShowFriendsUI(int32 LocalUserNum)
 {
 	EOS_UI_ShowFriendsOptions Options = {};
-	Options.ApiVersion = EOS_UI_SHOWFRIENDS_API_LATEST;
+	Options.ApiVersion = 1;
+	UE_EOS_CHECK_API_MISMATCH(EOS_UI_SHOWFRIENDS_API_LATEST, 1);
 	Options.LocalUserId = GetLocalEpicAccountId(LocalUserNum);
 
 	FOnShowFriendsCallback* CallbackObj = new FOnShowFriendsCallback(AsWeak());
@@ -1780,7 +1811,8 @@ void FUserManagerEOS::AddFriend(int32 LocalUserNum, EOS_EpicAccountId EpicAccoun
 	LocalUserNumToFriendsListMap[LocalUserNum]->Add(NetId, FriendRef);
 
 	EOS_Friends_GetStatusOptions Options = { };
-	Options.ApiVersion = EOS_FRIENDS_GETSTATUS_API_LATEST;
+	Options.ApiVersion = 1;
+	UE_EOS_CHECK_API_MISMATCH(EOS_FRIENDS_GETSTATUS_API_LATEST, 1);
 	Options.LocalUserId = UserNumToAccountIdMap[LocalUserNum];
 	Options.TargetUserId = EpicAccountId;
 	EOS_EFriendsStatus Status = EOS_Friends_GetStatus(EOSSubsystem->FriendsHandle, &Options);
@@ -1893,7 +1925,8 @@ bool FUserManagerEOS::ReadFriendsList(int32 LocalUserNum, const FString& ListNam
 	}
 
 	EOS_Friends_QueryFriendsOptions Options = { };
-	Options.ApiVersion = EOS_FRIENDS_QUERYFRIENDS_API_LATEST;
+	Options.ApiVersion = 1;
+	UE_EOS_CHECK_API_MISMATCH(EOS_FRIENDS_QUERYFRIENDS_API_LATEST, 1);
 	Options.LocalUserId = UserNumToAccountIdMap[LocalUserNum];
 
 	FReadFriendsCallback* CallbackObj = new FReadFriendsCallback(AsWeak());
@@ -1910,7 +1943,8 @@ bool FUserManagerEOS::ReadFriendsList(int32 LocalUserNum, const FString& ListNam
 		if (bWasSuccessful)
 		{
 			EOS_Friends_GetFriendsCountOptions Options = { };
-			Options.ApiVersion = EOS_FRIENDS_GETFRIENDSCOUNT_API_LATEST;
+			Options.ApiVersion = 1;
+			UE_EOS_CHECK_API_MISMATCH(EOS_FRIENDS_GETFRIENDSCOUNT_API_LATEST, 1);
 			Options.LocalUserId = UserNumToAccountIdMap[LocalUserNum];
 			int32 FriendCount = EOS_Friends_GetFriendsCount(EOSSubsystem->FriendsHandle, &Options);
 
@@ -1922,7 +1956,8 @@ bool FUserManagerEOS::ReadFriendsList(int32 LocalUserNum, const FString& ListNam
 			for (int32 Index = 0; Index < FriendCount; Index++)
 			{
 				EOS_Friends_GetFriendAtIndexOptions FriendIndexOptions = { };
-				FriendIndexOptions.ApiVersion = EOS_FRIENDS_GETFRIENDATINDEX_API_LATEST;
+				FriendIndexOptions.ApiVersion = 1;
+				UE_EOS_CHECK_API_MISMATCH(EOS_FRIENDS_GETFRIENDATINDEX_API_LATEST, 1);
 				FriendIndexOptions.Index = Index;
 				FriendIndexOptions.LocalUserId = Options.LocalUserId;
 				EOS_EpicAccountId FriendEpicAccountId = EOS_Friends_GetFriendAtIndex(EOSSubsystem->FriendsHandle, &FriendIndexOptions);
@@ -2047,7 +2082,8 @@ bool FUserManagerEOS::SendInvite(int32 LocalUserNum, const FUniqueNetId& FriendI
 	};
 
 	EOS_Friends_SendInviteOptions Options = { };
-	Options.ApiVersion = EOS_FRIENDS_SENDINVITE_API_LATEST;
+	Options.ApiVersion = 1;
+	UE_EOS_CHECK_API_MISMATCH(EOS_FRIENDS_SENDINVITE_API_LATEST, 1);
 	Options.LocalUserId = UserNumToAccountIdMap[LocalUserNum];
 	Options.TargetUserId = AccountId;
 	EOS_Friends_SendInvite(EOSSubsystem->FriendsHandle, &Options, CallbackObj, CallbackObj->GetCallbackPtr());
@@ -2090,7 +2126,8 @@ bool FUserManagerEOS::AcceptInvite(int32 LocalUserNum, const FUniqueNetId& Frien
 	};
 
 	EOS_Friends_AcceptInviteOptions Options = { };
-	Options.ApiVersion = EOS_FRIENDS_ACCEPTINVITE_API_LATEST;
+	Options.ApiVersion = 1;
+	UE_EOS_CHECK_API_MISMATCH(EOS_FRIENDS_ACCEPTINVITE_API_LATEST, 1);
 	Options.LocalUserId = UserNumToAccountIdMap[LocalUserNum];
 	Options.TargetUserId = AccountId;
 	EOS_Friends_AcceptInvite(EOSSubsystem->FriendsHandle, &Options, CallbackObj, CallbackObj->GetCallbackPtr());
@@ -2119,7 +2156,8 @@ bool FUserManagerEOS::RejectInvite(int32 LocalUserNum, const FUniqueNetId& Frien
 	}
 
 	EOS_Friends_RejectInviteOptions Options{ 0 };
-	Options.ApiVersion = EOS_FRIENDS_REJECTINVITE_API_LATEST;
+	Options.ApiVersion = 1;
+	UE_EOS_CHECK_API_MISMATCH(EOS_FRIENDS_REJECTINVITE_API_LATEST, 1);
 	Options.LocalUserId = UserNumToAccountIdMap[LocalUserNum];
 	Options.TargetUserId = AccountId;
 	EOS_Friends_RejectInvite(EOSSubsystem->FriendsHandle, &Options, nullptr, &EOSRejectInviteCallback);
@@ -2419,7 +2457,8 @@ struct FRichTextOptions :
 	FRichTextOptions() :
 		EOS_PresenceModification_SetRawRichTextOptions()
 	{
-		ApiVersion = EOS_PRESENCE_SETRAWRICHTEXT_API_LATEST;
+		ApiVersion = 1;
+		UE_EOS_CHECK_API_MISMATCH(EOS_PRESENCE_SETRAWRICHTEXT_API_LATEST, 1);
 		RichText = RichTextAnsi;
 	}
 	char RichTextAnsi[EOS_PRESENCE_RICH_TEXT_MAX_VALUE_LENGTH];
@@ -2439,7 +2478,8 @@ void FUserManagerEOS::SetPresence(const FUniqueNetId& UserId, const FOnlineUserP
 
 	EOS_HPresenceModification ChangeHandle = nullptr;
 	EOS_Presence_CreatePresenceModificationOptions Options = { };
-	Options.ApiVersion = EOS_PRESENCE_CREATEPRESENCEMODIFICATION_API_LATEST;
+	Options.ApiVersion = 1;
+	UE_EOS_CHECK_API_MISMATCH(EOS_PRESENCE_CREATEPRESENCEMODIFICATION_API_LATEST, 1);
 	Options.LocalUserId = AccountId;
 	EOS_Presence_CreatePresenceModification(EOSSubsystem->PresenceHandle, &Options, &ChangeHandle);
 	if (ChangeHandle == nullptr)
@@ -2449,7 +2489,8 @@ void FUserManagerEOS::SetPresence(const FUniqueNetId& UserId, const FOnlineUserP
 	}
 
 	EOS_PresenceModification_SetStatusOptions StatusOptions = { };
-	StatusOptions.ApiVersion = EOS_PRESENCE_SETSTATUS_API_LATEST;
+	StatusOptions.ApiVersion = 1;
+	UE_EOS_CHECK_API_MISMATCH(EOS_PRESENCE_SETSTATUS_API_LATEST, 1);
 	StatusOptions.Status = ToEOS_Presence_EStatus(Status.State);
 	EOS_EResult SetStatusResult = EOS_PresenceModification_SetStatus(ChangeHandle, &StatusOptions);
 	if (SetStatusResult != EOS_EResult::EOS_Success)
@@ -2475,12 +2516,14 @@ void FUserManagerEOS::SetPresence(const FUniqueNetId& UserId, const FOnlineUserP
 		const FPresenceStrings& RawString = RawStrings.Emplace_GetRef(It.Key(), It.Value().ToString());
 
 		EOS_Presence_DataRecord& Record = Records.Emplace_GetRef();
-		Record.ApiVersion = EOS_PRESENCE_DATARECORD_API_LATEST;
+		Record.ApiVersion = 1;
+		UE_EOS_CHECK_API_MISMATCH(EOS_PRESENCE_DATARECORD_API_LATEST, 1);
 		Record.Key = RawString.Key.Get();
 		Record.Value = RawString.Value.Get();
 	}
 	EOS_PresenceModification_SetDataOptions DataOptions = { };
-	DataOptions.ApiVersion = EOS_PRESENCE_SETDATA_API_LATEST;
+	DataOptions.ApiVersion = 1;
+	UE_EOS_CHECK_API_MISMATCH(EOS_PRESENCE_SETDATA_API_LATEST, 1);
 	DataOptions.RecordsCount = Records.Num();
 	DataOptions.Records = Records.GetData();
 	EOS_EResult SetDataResult = EOS_PresenceModification_SetData(ChangeHandle, &DataOptions);
@@ -2503,7 +2546,8 @@ void FUserManagerEOS::SetPresence(const FUniqueNetId& UserId, const FOnlineUserP
 	};
 
 	EOS_Presence_SetPresenceOptions PresOptions = { };
-	PresOptions.ApiVersion = EOS_PRESENCE_SETPRESENCE_API_LATEST;
+	PresOptions.ApiVersion = 1;
+	UE_EOS_CHECK_API_MISMATCH(EOS_PRESENCE_SETPRESENCE_API_LATEST, 1);
 	PresOptions.LocalUserId = AccountId;
 	PresOptions.PresenceModificationHandle = ChangeHandle;
 	// Last step commit the changes
@@ -2532,7 +2576,8 @@ void FUserManagerEOS::QueryPresence(const FUniqueNetId& UserId, const FOnPresenc
 	}
 
 	EOS_Presence_HasPresenceOptions HasOptions = { };
-	HasOptions.ApiVersion = EOS_PRESENCE_HASPRESENCE_API_LATEST;
+	HasOptions.ApiVersion = 1;
+	UE_EOS_CHECK_API_MISMATCH(EOS_PRESENCE_HASPRESENCE_API_LATEST, 1);
 	HasOptions.LocalUserId = UserNumToAccountIdMap[DefaultLocalUser];
 	HasOptions.TargetUserId = AccountId;
 	EOS_Bool bHasPresence = EOS_Presence_HasPresence(EOSSubsystem->PresenceHandle, &HasOptions);
@@ -2556,7 +2601,8 @@ void FUserManagerEOS::QueryPresence(const FUniqueNetId& UserId, const FOnPresenc
 
 		// Query for updated presence
 		EOS_Presence_QueryPresenceOptions Options = { };
-		Options.ApiVersion = EOS_PRESENCE_QUERYPRESENCE_API_LATEST;
+		Options.ApiVersion = 1;
+		UE_EOS_CHECK_API_MISMATCH(EOS_PRESENCE_QUERYPRESENCE_API_LATEST, 1);
 		Options.LocalUserId = HasOptions.LocalUserId;
 		Options.TargetUserId = HasOptions.TargetUserId;
 		EOS_Presence_QueryPresence(EOSSubsystem->PresenceHandle, &Options, CallbackObj, CallbackObj->GetCallbackPtr());
@@ -2573,7 +2619,8 @@ void FUserManagerEOS::UpdatePresence(EOS_EpicAccountId AccountId)
 {
 	EOS_Presence_Info* PresenceInfo = nullptr;
 	EOS_Presence_CopyPresenceOptions Options = { };
-	Options.ApiVersion = EOS_PRESENCE_COPYPRESENCE_API_LATEST;
+	Options.ApiVersion = 3;
+	UE_EOS_CHECK_API_MISMATCH(EOS_PRESENCE_COPYPRESENCE_API_LATEST, 3);
 	Options.LocalUserId = UserNumToAccountIdMap[DefaultLocalUser];
 	Options.TargetUserId = AccountId;
 	EOS_EResult CopyResult = EOS_Presence_CopyPresence(EOSSubsystem->PresenceHandle, &Options, &PresenceInfo);
@@ -2708,7 +2755,8 @@ void FUserManagerEOS::ReadUserInfo(int32 LocalUserNum, EOS_EpicAccountId EpicAcc
 	};
 
 	EOS_UserInfo_QueryUserInfoOptions Options = { };
-	Options.ApiVersion = EOS_USERINFO_QUERYUSERINFO_API_LATEST;
+	Options.ApiVersion = 1;
+	UE_EOS_CHECK_API_MISMATCH(EOS_USERINFO_QUERYUSERINFO_API_LATEST, 1);
 	Options.LocalUserId = UserNumToAccountIdMap[DefaultLocalUser];
 	Options.TargetUserId = EpicAccountId;
 	EOS_UserInfo_QueryUserInfo(EOSSubsystem->UserInfoHandle, &Options, CallbackObj, CallbackObj->GetCallbackPtr());
@@ -2753,7 +2801,8 @@ struct FQueryByDisplayNameOptions :
 	FQueryByDisplayNameOptions() :
 		EOS_UserInfo_QueryUserInfoByDisplayNameOptions()
 	{
-		ApiVersion = EOS_USERINFO_QUERYUSERINFOBYDISPLAYNAME_API_LATEST;
+		ApiVersion = 1;
+		UE_EOS_CHECK_API_MISMATCH(EOS_USERINFO_QUERYUSERINFOBYDISPLAYNAME_API_LATEST, 1);
 		DisplayName = DisplayNameAnsi;
 	}
 	char DisplayNameAnsi[EOS_OSS_STRING_BUFFER_LENGTH];
@@ -2823,7 +2872,8 @@ struct FQueryByStringIdsOptions :
 		{
 			PointerArray[Index] = new char[EOS_CONNECT_EXTERNAL_ACCOUNT_ID_MAX_LENGTH+1];
 		}
-		ApiVersion = EOS_CONNECT_QUERYEXTERNALACCOUNTMAPPINGS_API_LATEST;
+		ApiVersion = 1;
+		UE_EOS_CHECK_API_MISMATCH(EOS_CONNECT_QUERYEXTERNALACCOUNTMAPPINGS_API_LATEST, 1);
 		AccountIdType = EOS_EExternalAccountType::EOS_EAT_EPIC;
 		ExternalAccountIds = (const char**)PointerArray.GetData();
 		ExternalAccountIdCount = InNumStringIds;
@@ -2846,7 +2896,8 @@ struct FGetAccountMappingOptions :
 	FGetAccountMappingOptions() :
 		EOS_Connect_GetExternalAccountMappingsOptions()
 	{
-		ApiVersion = EOS_CONNECT_GETEXTERNALACCOUNTMAPPINGS_API_LATEST;
+		ApiVersion = 1;
+		UE_EOS_CHECK_API_MISMATCH(EOS_CONNECT_GETEXTERNALACCOUNTMAPPINGS_API_LATEST, 1);
 		AccountIdType = EOS_EExternalAccountType::EOS_EAT_EPIC;
 		TargetExternalUserId = AccountId;
 	}
diff --git a/Engine/Plugins/Online/OnlineSubsystemEOS/Source/OnlineSubsystemEOS/Private/Windows/WindowsEOSHelpers.cpp b/Engine/Plugins/Online/OnlineSubsystemEOS/Source/OnlineSubsystemEOS/Private/Windows/WindowsEOSHelpers.cpp
--- a/Engine/Plugins/Online/OnlineSubsystemEOS/Source/OnlineSubsystemEOS/Private/Windows/WindowsEOSHelpers.cpp
+++ b/Engine/Plugins/Online/OnlineSubsystemEOS/Source/OnlineSubsystemEOS/Private/Windows/WindowsEOSHelpers.cpp
@@ -17,8 +17,8 @@ IEOSPlatformHandlePtr FWindowsEOSHelpers::CreatePlatform(EOS_Platform_Options& P
 	const FTCHARToUTF8 Utf8XAudioPath(*IFileManager::Get().ConvertToAbsolutePathForExternalAppForRead(*FPaths::Combine(FPaths::EngineDir(), TEXT("Binaries/ThirdParty/Windows/XAudio2_9"), PLATFORM_64BITS ? TEXT("x64") : TEXT("x86"), TEXT("xaudio2_9redist.dll"))));
 
 	EOS_Windows_RTCOptions WindowsRtcOptions = { 0 };
-	WindowsRtcOptions.ApiVersion = EOS_WINDOWS_RTCOPTIONS_API_LATEST;
-	static_assert(EOS_WINDOWS_RTCOPTIONS_API_LATEST == 1, "EOS_Windows_RTCOptions updated, check new fields");
+	WindowsRtcOptions.ApiVersion = 1;
+	UE_EOS_CHECK_API_MISMATCH(EOS_WINDOWS_RTCOPTIONS_API_LATEST, 1);
 	WindowsRtcOptions.XAudio29DllPath = Utf8XAudioPath.Get();
 
 	EOS_Platform_RTCOptions* RTCOptions = const_cast<EOS_Platform_RTCOptions*>(PlatformOptions.RTCOptions);
diff --git a/Engine/Plugins/Online/SocketSubsystemEOS/Source/SocketSubsystemEOS/Private/SocketEOS.cpp b/Engine/Plugins/Online/SocketSubsystemEOS/Source/SocketSubsystemEOS/Private/SocketEOS.cpp
--- a/Engine/Plugins/Online/SocketSubsystemEOS/Source/SocketSubsystemEOS/Private/SocketEOS.cpp
+++ b/Engine/Plugins/Online/SocketSubsystemEOS/Source/SocketSubsystemEOS/Private/SocketEOS.cpp
@@ -3,6 +3,7 @@
 #include "SocketEOS.h"
 #include "SocketTypes.h"
 #include "SocketSubsystemEOS.h"
+#include "EOSShared.h"
 
 #if WITH_EOS_SDK
 	#include "eos_p2p.h"
@@ -107,11 +108,13 @@ bool FSocketEOS::Close()
 	if (LocalAddress.IsValid())
 	{
 		EOS_P2P_SocketId SocketId = { };
-		SocketId.ApiVersion = EOS_P2P_SOCKETID_API_LATEST;
+		SocketId.ApiVersion = 1;
+		UE_EOS_CHECK_API_MISMATCH(EOS_P2P_SOCKETID_API_LATEST, 1);
 		FCStringAnsi::Strcpy(SocketId.SocketName, LocalAddress.GetSocketName());
 
 		EOS_P2P_CloseConnectionsOptions Options = { };
-		Options.ApiVersion = EOS_P2P_CLOSECONNECTIONS_API_LATEST;
+		Options.ApiVersion = 1;
+		UE_EOS_CHECK_API_MISMATCH(EOS_P2P_CLOSECONNECTIONS_API_LATEST, 1);
 		Options.LocalUserId = SocketSubsystem.GetLocalUserId();
 		Options.SocketId = &SocketId;
 
@@ -195,11 +198,13 @@ bool FSocketEOS::Listen(int32)
 #if WITH_EOS_SDK
 	// Add listener for inbound connections
 	EOS_P2P_SocketId SocketId = { };
-	SocketId.ApiVersion = EOS_P2P_SOCKETID_API_LATEST;
+	SocketId.ApiVersion = 1;
+	UE_EOS_CHECK_API_MISMATCH(EOS_P2P_SOCKETID_API_LATEST, 1);
 	FCStringAnsi::Strcpy(SocketId.SocketName, LocalAddress.GetSocketName());
 
 	EOS_P2P_AddNotifyPeerConnectionRequestOptions Options = { };
-	Options.ApiVersion = EOS_P2P_ADDNOTIFYPEERCONNECTIONREQUEST_API_LATEST;
+	Options.ApiVersion = 1;
+	UE_EOS_CHECK_API_MISMATCH(EOS_P2P_ADDNOTIFYPEERCONNECTIONREQUEST_API_LATEST, 1);
 	Options.LocalUserId = LocalAddress.GetLocalUserId();
 	Options.SocketId = &SocketId;
 
@@ -222,11 +227,13 @@ bool FSocketEOS::Listen(int32)
 			ClosedRemotes.Remove(RemoteAddress);
 
 			EOS_P2P_SocketId SocketId = { };
-			SocketId.ApiVersion = EOS_P2P_SOCKETID_API_LATEST;
+			SocketId.ApiVersion = 1;
+			UE_EOS_CHECK_API_MISMATCH(EOS_P2P_SOCKETID_API_LATEST, 1);
 			FCStringAnsi::Strcpy(SocketId.SocketName, Info->SocketId->SocketName);
 
 			EOS_P2P_AcceptConnectionOptions Options = { };
-			Options.ApiVersion = EOS_P2P_ACCEPTCONNECTION_API_LATEST;
+			Options.ApiVersion = 1;
+			UE_EOS_CHECK_API_MISMATCH(EOS_P2P_ACCEPTCONNECTION_API_LATEST, 1);
 			Options.LocalUserId = LocalAddress.GetLocalUserId();
 			Options.RemoteUserId = Info->RemoteUserId;
 			Options.SocketId = &SocketId;
@@ -273,7 +280,8 @@ bool FSocketEOS::HasPendingData(uint32& PendingDataSize)
 
 #if WITH_EOS_SDK
 	EOS_P2P_GetNextReceivedPacketSizeOptions Options = { };
-	Options.ApiVersion = EOS_P2P_GETNEXTRECEIVEDPACKETSIZE_API_LATEST;
+	Options.ApiVersion = 2;
+	UE_EOS_CHECK_API_MISMATCH(EOS_P2P_GETNEXTRECEIVEDPACKETSIZE_API_LATEST, 2);
 	Options.LocalUserId = LocalAddress.GetLocalUserId();
 	uint8 Channel = LocalAddress.GetChannel();
 	Options.RequestedChannel = &Channel;
@@ -383,11 +391,13 @@ bool FSocketEOS::SendTo(const uint8* Data, int32 Count, int32& OutBytesSent, con
 	RegisterClosedNotification();
 
 	EOS_P2P_SocketId SocketId = { };
-	SocketId.ApiVersion = EOS_P2P_SOCKETID_API_LATEST;
+	SocketId.ApiVersion = 1;
+	UE_EOS_CHECK_API_MISMATCH(EOS_P2P_SOCKETID_API_LATEST, 1);
 	FCStringAnsi::Strcpy(SocketId.SocketName, DestinationAddress.GetSocketName());
 
 	EOS_P2P_SendPacketOptions Options = { };
-	Options.ApiVersion = EOS_P2P_SENDPACKET_API_LATEST;
+	Options.ApiVersion = 3;
+	UE_EOS_CHECK_API_MISMATCH(EOS_P2P_SENDPACKET_API_LATEST, 3);
 	Options.LocalUserId = LocalAddress.GetLocalUserId();
 	Options.RemoteUserId = DestinationAddress.GetRemoteUserId();
 	Options.SocketId = &SocketId;
@@ -445,7 +455,8 @@ bool FSocketEOS::RecvFrom(uint8* Data, int32 BufferSize, int32& BytesRead, FInte
 
 #if WITH_EOS_SDK
 	EOS_P2P_ReceivePacketOptions Options = { };
-	Options.ApiVersion = EOS_P2P_RECEIVEPACKET_API_LATEST;
+	Options.ApiVersion = 2;
+	UE_EOS_CHECK_API_MISMATCH(EOS_P2P_RECEIVEPACKET_API_LATEST, 2);
 	Options.LocalUserId = LocalAddress.GetLocalUserId();
 	Options.MaxDataSizeBytes = BufferSize;
 	uint8 Channel = LocalAddress.GetChannel();
@@ -621,11 +632,13 @@ bool FSocketEOS::Close(const FInternetAddrEOS& RemoteAddress)
 	ClosedRemotes.Add(RemoteAddress);
 
 	EOS_P2P_SocketId SocketId = { };
-	SocketId.ApiVersion = EOS_P2P_SOCKETID_API_LATEST;
+	SocketId.ApiVersion = 1;
+	UE_EOS_CHECK_API_MISMATCH(EOS_P2P_SOCKETID_API_LATEST, 1);
 	FCStringAnsi::Strcpy(SocketId.SocketName, RemoteAddress.GetSocketName());
 
 	EOS_P2P_CloseConnectionOptions Options = { };
-	Options.ApiVersion = EOS_P2P_CLOSECONNECTION_API_LATEST;
+	Options.ApiVersion = 1;
+	UE_EOS_CHECK_API_MISMATCH(EOS_P2P_CLOSECONNECTION_API_LATEST, 1);
 	Options.LocalUserId = LocalAddress.GetLocalUserId();
 	Options.RemoteUserId = RemoteAddress.GetRemoteUserId();
 	Options.SocketId = &SocketId;
@@ -656,11 +669,13 @@ void FSocketEOS::RegisterClosedNotification()
 	}
 	
 	EOS_P2P_SocketId SocketId = { };
-	SocketId.ApiVersion = EOS_P2P_SOCKETID_API_LATEST;
+	SocketId.ApiVersion = 1;
+	UE_EOS_CHECK_API_MISMATCH(EOS_P2P_SOCKETID_API_LATEST, 1);
 	FCStringAnsi::Strcpy(SocketId.SocketName, LocalAddress.GetSocketName());
 
 	EOS_P2P_AddNotifyPeerConnectionClosedOptions Options = { };
-	Options.ApiVersion = EOS_P2P_ADDNOTIFYPEERCONNECTIONCLOSED_API_LATEST;
+	Options.ApiVersion = 1;
+	UE_EOS_CHECK_API_MISMATCH(EOS_P2P_ADDNOTIFYPEERCONNECTIONCLOSED_API_LATEST, 1);
 	Options.LocalUserId = LocalAddress.GetLocalUserId();
 	Options.SocketId = &SocketId;
 
diff --git a/Engine/Plugins/Online/SocketSubsystemEOS/Source/SocketSubsystemEOS/Private/SocketSubsystemEOS.cpp b/Engine/Plugins/Online/SocketSubsystemEOS/Source/SocketSubsystemEOS/Private/SocketSubsystemEOS.cpp
--- a/Engine/Plugins/Online/SocketSubsystemEOS/Source/SocketSubsystemEOS/Private/SocketSubsystemEOS.cpp
+++ b/Engine/Plugins/Online/SocketSubsystemEOS/Source/SocketSubsystemEOS/Private/SocketSubsystemEOS.cpp
@@ -10,6 +10,7 @@
 #include "Modules/ModuleManager.h"
 #include "Misc/OutputDeviceRedirector.h"
 #include "OnlineSubsystemUtils.h"
+#include "EOSShared.h"
 
 #if WITH_EOS_SDK
 	#include "eos_sdk.h"
@@ -99,8 +100,8 @@ bool FSocketSubsystemEOS::Init(FString& Error)
 	if (!RelayControlStr.IsEmpty())
 	{
 		EOS_P2P_SetRelayControlOptions Options = {};
-		Options.ApiVersion = EOS_P2P_SETRELAYCONTROL_API_LATEST;
-		static_assert(EOS_P2P_SETRELAYCONTROL_API_LATEST == 1, "EOS_P2P_SetRelayControlOptions updated, check new fields");
+		Options.ApiVersion = 1;
+		UE_EOS_CHECK_API_MISMATCH(EOS_P2P_SETRELAYCONTROL_API_LATEST, 1);
 
 		Options.RelayControl = LexFromString(RelayControlStr);
 
diff --git a/Engine/Plugins/Online/VoiceChat/EOSVoiceChat/Source/EOSVoiceChat/Private/EOSAudioDevicePool.cpp b/Engine/Plugins/Online/VoiceChat/EOSVoiceChat/Source/EOSVoiceChat/Private/EOSAudioDevicePool.cpp
--- a/Engine/Plugins/Online/VoiceChat/EOSVoiceChat/Source/EOSVoiceChat/Private/EOSAudioDevicePool.cpp
+++ b/Engine/Plugins/Online/VoiceChat/EOSVoiceChat/Source/EOSVoiceChat/Private/EOSAudioDevicePool.cpp
@@ -89,16 +89,16 @@ TArray<FVoiceChatDeviceInfo> FEOSAudioDevicePool::GetRtcInputDeviceInfos(int32&
 	EOS_HRTCAudio RTCAudioHandle = EOS_RTC_GetAudioInterface(EosRtcInterface);
 
 	EOS_RTCAudio_GetAudioInputDevicesCountOptions CountOptions = {};
-	CountOptions.ApiVersion = EOS_RTCAUDIO_GETAUDIOINPUTDEVICESCOUNT_API_LATEST;
-	static_assert(EOS_RTCAUDIO_GETAUDIOINPUTDEVICESCOUNT_API_LATEST == 1, "EOS_RTCAudio_GetAudioInputDevicesCountOptions updated, check new fields");
+	CountOptions.ApiVersion = 1;
+	UE_EOS_CHECK_API_MISMATCH(EOS_RTCAUDIO_GETAUDIOINPUTDEVICESCOUNT_API_LATEST, 1);
 
 	uint32_t Count = EOS_RTCAudio_GetAudioInputDevicesCount(RTCAudioHandle, &CountOptions);
 
 	for (uint32_t Index = 0; Index < Count; Index++)
 	{
 		EOS_RTCAudio_GetAudioInputDeviceByIndexOptions GetByIndexOptions = {};
-		GetByIndexOptions.ApiVersion = EOS_RTCAUDIO_GETAUDIOINPUTDEVICEBYINDEX_API_LATEST;
-		static_assert(EOS_RTCAUDIO_GETAUDIOINPUTDEVICEBYINDEX_API_LATEST == 1, "EOS_RTCAudio_GetAudioInputDeviceByIndexOptions updated, check new fields");
+		GetByIndexOptions.ApiVersion = 1;
+		UE_EOS_CHECK_API_MISMATCH(EOS_RTCAUDIO_GETAUDIOINPUTDEVICEBYINDEX_API_LATEST, 1);
 		GetByIndexOptions.DeviceInfoIndex = Index;
 		if (const EOS_RTCAudio_AudioInputDeviceInfo* DeviceInfo = EOS_RTCAudio_GetAudioInputDeviceByIndex(RTCAudioHandle, &GetByIndexOptions))
 		{
@@ -135,16 +135,16 @@ TArray<FVoiceChatDeviceInfo> FEOSAudioDevicePool::GetRtcOutputDeviceInfos(int32&
 	EOS_HRTCAudio RTCAudioHandle = EOS_RTC_GetAudioInterface(EosRtcInterface);
 
 	EOS_RTCAudio_GetAudioOutputDevicesCountOptions CountOptions = {};
-	CountOptions.ApiVersion = EOS_RTCAUDIO_GETAUDIOOUTPUTDEVICESCOUNT_API_LATEST;
-	static_assert(EOS_RTCAUDIO_GETAUDIOOUTPUTDEVICESCOUNT_API_LATEST == 1, "EOS_RTCAudio_GetAudioOutputDevicesCountOptions updated, check new fields");
+	CountOptions.ApiVersion = 1;
+	UE_EOS_CHECK_API_MISMATCH(EOS_RTCAUDIO_GETAUDIOOUTPUTDEVICESCOUNT_API_LATEST, 1);
 
 	uint32_t Count = EOS_RTCAudio_GetAudioOutputDevicesCount(RTCAudioHandle, &CountOptions);
 
 	for (uint32_t Index = 0; Index < Count; Index++)
 	{
 		EOS_RTCAudio_GetAudioOutputDeviceByIndexOptions GetByIndexOptions = {};
-		GetByIndexOptions.ApiVersion = EOS_RTCAUDIO_GETAUDIOOUTPUTDEVICEBYINDEX_API_LATEST;
-		static_assert(EOS_RTCAUDIO_GETAUDIOOUTPUTDEVICEBYINDEX_API_LATEST == 1, "EOS_RTCAudio_GetAudioOutputDeviceByIndexOptions updated, check new fields");
+		GetByIndexOptions.ApiVersion = 1;
+		UE_EOS_CHECK_API_MISMATCH(EOS_RTCAUDIO_GETAUDIOOUTPUTDEVICEBYINDEX_API_LATEST, 1);
 		GetByIndexOptions.DeviceInfoIndex = Index;
 		if (const EOS_RTCAudio_AudioOutputDeviceInfo* DeviceInfo = EOS_RTCAudio_GetAudioOutputDeviceByIndex(RTCAudioHandle, &GetByIndexOptions))
 		{
diff --git a/Engine/Plugins/Online/VoiceChat/EOSVoiceChat/Source/EOSVoiceChat/Private/EOSVoiceChat.cpp b/Engine/Plugins/Online/VoiceChat/EOSVoiceChat/Source/EOSVoiceChat/Private/EOSVoiceChat.cpp
--- a/Engine/Plugins/Online/VoiceChat/EOSVoiceChat/Source/EOSVoiceChat/Private/EOSVoiceChat.cpp
+++ b/Engine/Plugins/Online/VoiceChat/EOSVoiceChat/Source/EOSVoiceChat/Private/EOSVoiceChat.cpp
@@ -134,8 +134,8 @@ void FEOSVoiceChat::Initialize(const FOnVoiceChatInitializeCompleteDelegate& Ini
 					const FTCHARToUTF8 Utf8OverrideLocaleCode(*ConfigOverrideLocaleCode);
 
 					EOS_Platform_Options PlatformOptions = {};
-					PlatformOptions.ApiVersion = EOS_PLATFORM_OPTIONS_API_LATEST;
-					static_assert(EOS_PLATFORM_OPTIONS_API_LATEST == 12, "EOS_Platform_Options updated, check new fields");
+					PlatformOptions.ApiVersion = 12;
+					UE_EOS_CHECK_API_MISMATCH(EOS_PLATFORM_OPTIONS_API_LATEST, 12);
 					PlatformOptions.Reserved = nullptr;
 					PlatformOptions.ProductId = ConfigProductId.IsEmpty() ? nullptr : Utf8ProductId.Get();
 					PlatformOptions.SandboxId = ConfigSandboxId.IsEmpty() ? nullptr : Utf8SandboxId.Get();
@@ -155,8 +155,8 @@ void FEOSVoiceChat::Initialize(const FOnVoiceChatInitializeCompleteDelegate& Ini
 #endif
 
 					EOS_Platform_RTCOptions PlatformRTCOptions = {};
-					PlatformRTCOptions.ApiVersion = EOS_PLATFORM_RTCOPTIONS_API_LATEST;
-					static_assert(EOS_PLATFORM_RTCOPTIONS_API_LATEST == 1, "EOS_Platform_RTCOptions updated, check new fields");
+					PlatformRTCOptions.ApiVersion = 1;
+					UE_EOS_CHECK_API_MISMATCH(EOS_PLATFORM_RTCOPTIONS_API_LATEST, 1);
 					PlatformOptions.RTCOptions = &PlatformRTCOptions;
 
 					EosPlatformHandle = EOSPlatformCreate(PlatformOptions);
@@ -783,8 +783,8 @@ FString FEOSVoiceChat::InsecureGetJoinToken(const FString& ChannelName, EVoiceCh
 void FEOSVoiceChat::BindInitCallbacks()
 {
 	EOS_RTCAudio_AddNotifyAudioDevicesChangedOptions AudioDevicesChangedOptions = {};
-	AudioDevicesChangedOptions.ApiVersion = EOS_RTCAUDIO_ADDNOTIFYAUDIODEVICESCHANGED_API_LATEST;
-	static_assert(EOS_RTCAUDIO_ADDNOTIFYAUDIODEVICESCHANGED_API_LATEST == 1, "EOS_RTC_AddNotifyAudioDevicesChangedOptions updated, check new fields");
+	AudioDevicesChangedOptions.ApiVersion = 1;
+	UE_EOS_CHECK_API_MISMATCH(EOS_RTCAUDIO_ADDNOTIFYAUDIODEVICESCHANGED_API_LATEST, 1);
 	InitSession.OnAudioDevicesChangedNotificationId = EOS_RTCAudio_AddNotifyAudioDevicesChanged(EOS_RTC_GetAudioInterface(GetRtcInterface()), &AudioDevicesChangedOptions, this, &FEOSVoiceChat::OnAudioDevicesChangedStatic);
 	if (InitSession.OnAudioDevicesChangedNotificationId == EOS_INVALID_NOTIFICATIONID)
 	{
diff --git a/Engine/Plugins/Online/VoiceChat/EOSVoiceChat/Source/EOSVoiceChat/Private/EOSVoiceChatUser.cpp b/Engine/Plugins/Online/VoiceChat/EOSVoiceChat/Source/EOSVoiceChat/Private/EOSVoiceChatUser.cpp
--- a/Engine/Plugins/Online/VoiceChat/EOSVoiceChat/Source/EOSVoiceChat/Private/EOSVoiceChatUser.cpp
+++ b/Engine/Plugins/Online/VoiceChat/EOSVoiceChat/Source/EOSVoiceChat/Private/EOSVoiceChatUser.cpp
@@ -581,8 +581,8 @@ void FEOSVoiceChatUser::JoinChannel(const FString& ChannelName, const FString& C
 			// Attempt to join the channel
 			{
 				EOS_RTC_JoinRoomOptions JoinOptions = {};
-				JoinOptions.ApiVersion = EOS_RTC_JOINROOM_API_LATEST;
-				static_assert(EOS_RTC_JOINROOM_API_LATEST == 1, "EOS_RTC_JoinRoomOptions updated, check new fields");
+				JoinOptions.ApiVersion = 1;
+				UE_EOS_CHECK_API_MISMATCH(EOS_RTC_JOINROOM_API_LATEST, 1);
 				JoinOptions.LocalUserId = LoginSession.LocalUserProductUserId;
 				JoinOptions.RoomName = Utf8RoomName.Get();
 				JoinOptions.ClientBaseUrl = Utf8ClientBaseUrl.Get();
@@ -933,8 +933,8 @@ bool FEOSVoiceChatUser::AddLobbyRoom(const FString& LobbyId)
 	const FTCHARToUTF8 Utf8LobbyId(*LobbyId);
 
 	EOS_Lobby_GetRTCRoomNameOptions GetRoomNameOptions = {};
-	GetRoomNameOptions.ApiVersion = EOS_LOBBY_GETRTCROOMNAME_API_LATEST;
-	static_assert(EOS_LOBBY_GETRTCROOMNAME_API_LATEST == 1, "EOS_Lobby_GetRTCRoomNameOptions updated, check new fields");
+	GetRoomNameOptions.ApiVersion = 1;
+	UE_EOS_CHECK_API_MISMATCH(EOS_LOBBY_GETRTCROOMNAME_API_LATEST, 1);
 	GetRoomNameOptions.LobbyId = Utf8LobbyId.Get();
 	GetRoomNameOptions.LocalUserId = LoginSession.LocalUserProductUserId;
 	char Utf8RoomName[256];
@@ -959,8 +959,8 @@ bool FEOSVoiceChatUser::AddLobbyRoom(const FString& LobbyId)
 		// Call UpdateReceiving once to set the default receiving state for all participants
 		{
 			EOS_RTCAudio_UpdateReceivingOptions UpdateReceivingOptions = {};
-			UpdateReceivingOptions.ApiVersion = EOS_RTCAUDIO_UPDATERECEIVING_API_LATEST;
-			static_assert(EOS_RTCAUDIO_UPDATERECEIVING_API_LATEST == 1, "EOS_RTCAudio_UpdateReceivingOptions updated, check new fields");
+			UpdateReceivingOptions.ApiVersion = 1;
+			UE_EOS_CHECK_API_MISMATCH(EOS_RTCAUDIO_UPDATERECEIVING_API_LATEST, 1);
 			UpdateReceivingOptions.LocalUserId = LoginSession.LocalUserProductUserId;
 			UpdateReceivingOptions.RoomName = Utf8RoomName;
 			UpdateReceivingOptions.ParticipantId = nullptr;
@@ -973,8 +973,8 @@ bool FEOSVoiceChatUser::AddLobbyRoom(const FString& LobbyId)
 		
 		// Get the current lobby connectedness.
 		EOS_Lobby_IsRTCRoomConnectedOptions IsRoomConnectedOptions = {};
-		IsRoomConnectedOptions.ApiVersion = EOS_LOBBY_ISRTCROOMCONNECTED_API_LATEST;
-		static_assert(EOS_LOBBY_ISRTCROOMCONNECTED_API_LATEST == 1, "EOS_Lobby_IsRTCRoomConnectedOptions updated, check new fields");
+		IsRoomConnectedOptions.ApiVersion = 1;
+		UE_EOS_CHECK_API_MISMATCH(EOS_LOBBY_ISRTCROOMCONNECTED_API_LATEST, 1);
 		IsRoomConnectedOptions.LobbyId = Utf8LobbyId.Get();
 		IsRoomConnectedOptions.LocalUserId = LoginSession.LocalUserProductUserId;
 		EOS_Bool bLobbyChannelConnected = false;
@@ -1041,8 +1041,8 @@ void FEOSVoiceChatUser::RtcRegisterUser(const FString& UserId, const FOnVoiceCha
 	const FTCHARToUTF8 Utf8UserId(*UserId);
 
 	EOS_RTCAudio_RegisterPlatformAudioUserOptions Options = {};
-	Options.ApiVersion = EOS_RTCAUDIO_REGISTERPLATFORMAUDIOUSER_API_LATEST;
-	static_assert(EOS_RTCAUDIO_REGISTERPLATFORMAUDIOUSER_API_LATEST == 1, "EOS_RTCAUDIO_REGISTERPLATFORMAUDIOUSER_API_LATEST updated, check new fields");
+	Options.ApiVersion = 1;
+	UE_EOS_CHECK_API_MISMATCH(EOS_RTCAUDIO_REGISTERPLATFORMAUDIOUSER_API_LATEST, 1);
 	Options.UserId = Utf8UserId.Get();
 
 	CSV_SCOPED_TIMING_STAT_EXCLUSIVE(EOSVoiceChat);
@@ -1061,8 +1061,8 @@ void FEOSVoiceChatUser::RtcUnregisterUser(const FString& UserId, const FOnVoiceC
 	const FTCHARToUTF8 Utf8UserId(*UserId);
 
 	EOS_RTCAudio_UnregisterPlatformAudioUserOptions Options = {};
-	Options.ApiVersion = EOS_RTCAUDIO_UNREGISTERPLATFORMAUDIOUSER_API_LATEST;
-	static_assert(EOS_RTCAUDIO_UNREGISTERPLATFORMAUDIOUSER_API_LATEST == 1, "EOS_RTCAUDIO_UNREGISTERPLATFORMAUDIOUSER_API_LATEST updated, check new fields");
+	Options.ApiVersion = 1;
+	UE_EOS_CHECK_API_MISMATCH(EOS_RTCAUDIO_UNREGISTERPLATFORMAUDIOUSER_API_LATEST, 1);
 	Options.UserId = Utf8UserId.Get();
 
 	CSV_SCOPED_TIMING_STAT_EXCLUSIVE(EOSVoiceChat);
@@ -1161,8 +1161,8 @@ void FEOSVoiceChatUser::ApplyAudioInputOptions()
 		const FTCHARToUTF8 Utf8DeviceId(DeviceIdPtr);
 
 		EOS_RTCAudio_SetAudioInputSettingsOptions Options = {};
-		Options.ApiVersion = EOS_RTCAUDIO_SETAUDIOINPUTSETTINGS_API_LATEST;
-		static_assert(EOS_RTCAUDIO_SETAUDIOINPUTSETTINGS_API_LATEST == 1, "EOS_RTCAudio_SetAudioInputSettingsOptions updated, check new fields");
+		Options.ApiVersion = 1;
+		UE_EOS_CHECK_API_MISMATCH(EOS_RTCAUDIO_SETAUDIOINPUTSETTINGS_API_LATEST, 1);
 		Options.LocalUserId = LoginSession.LocalUserProductUserId;
 		Options.DeviceId = Utf8DeviceId.Get();
 		Options.Volume = AudioInputOptions.Volume * 50.0;
@@ -1187,8 +1187,8 @@ void FEOSVoiceChatUser::ApplyAudioOutputOptions()
 		const FTCHARToUTF8 Utf8DeviceId(DeviceIdPtr);
 
 		EOS_RTCAudio_SetAudioOutputSettingsOptions Options = {};
-		Options.ApiVersion = EOS_RTCAUDIO_SETAUDIOOUTPUTSETTINGS_API_LATEST;
-		static_assert(EOS_RTCAUDIO_SETAUDIOOUTPUTSETTINGS_API_LATEST == 1, "EOS_RTCAudio_SetAudioOutputSettingsOptions updated, check new fields");
+		Options.ApiVersion = 1;
+		UE_EOS_CHECK_API_MISMATCH(EOS_RTCAUDIO_SETAUDIOOUTPUTSETTINGS_API_LATEST, 1);
 		Options.LocalUserId = LoginSession.LocalUserProductUserId;
 		Options.DeviceId = Utf8DeviceId.Get();
 		Options.Volume = AudioOutputOptions.bMuted ? 0.0f : AudioOutputOptions.Volume * 50.0;
@@ -1211,8 +1211,8 @@ void FEOSVoiceChatUser::ApplyPlayerBlock(const FGlobalParticipant& GlobalPartici
 	const EOS_ProductUserId ParticipantId = EOSProductUserIdFromString(*GlobalParticipant.PlayerName);
 
 	EOS_RTC_BlockParticipantOptions Options = {};
-	Options.ApiVersion = EOS_RTC_BLOCKPARTICIPANT_API_LATEST;
-	static_assert(EOS_RTC_BLOCKPARTICIPANT_API_LATEST == 1, "EOS_RTC_BlockParticipantOptions updated, check new fields");
+	Options.ApiVersion = 1;
+	UE_EOS_CHECK_API_MISMATCH(EOS_RTC_BLOCKPARTICIPANT_API_LATEST, 1);
 	Options.LocalUserId = LoginSession.LocalUserProductUserId;
 	Options.RoomName = Utf8RoomName.Get();
 	Options.ParticipantId = ParticipantId;
@@ -1228,8 +1228,8 @@ void FEOSVoiceChatUser::ApplyReceivingOptions(const FChannelSession& ChannelSess
 	const FTCHARToUTF8 Utf8RoomName(*ChannelSession.ChannelName);
 
 	EOS_RTCAudio_UpdateReceivingOptions UpdateReceivingOptions = {};
-	UpdateReceivingOptions.ApiVersion = EOS_RTCAUDIO_UPDATERECEIVING_API_LATEST;
-	static_assert(EOS_RTCAUDIO_UPDATERECEIVING_API_LATEST == 1, "EOS_RTCAudio_UpdateReceivingOptions updated, check new fields");
+	UpdateReceivingOptions.ApiVersion = 1;
+	UE_EOS_CHECK_API_MISMATCH(EOS_RTCAUDIO_UPDATERECEIVING_API_LATEST, 1);
 	UpdateReceivingOptions.LocalUserId = LoginSession.LocalUserProductUserId;
 	UpdateReceivingOptions.RoomName = Utf8RoomName.Get();
 	UpdateReceivingOptions.ParticipantId = nullptr;
@@ -1269,8 +1269,8 @@ void FEOSVoiceChatUser::ApplyPlayerReceivingOptions(const FGlobalParticipant& Gl
 
 	{
 		EOS_RTCAudio_UpdateReceivingOptions UpdateReceivingOptions = {};
-		UpdateReceivingOptions.ApiVersion = EOS_RTCAUDIO_UPDATERECEIVING_API_LATEST;
-		static_assert(EOS_RTCAUDIO_UPDATERECEIVING_API_LATEST == 1, "EOS_RTCAudio_UpdateReceivingOptions updated, check new fields");
+		UpdateReceivingOptions.ApiVersion = 1;
+		UE_EOS_CHECK_API_MISMATCH(EOS_RTCAUDIO_UPDATERECEIVING_API_LATEST, 1);
 		UpdateReceivingOptions.LocalUserId = LoginSession.LocalUserProductUserId;
 		UpdateReceivingOptions.RoomName = Utf8RoomName.Get();
 		UpdateReceivingOptions.ParticipantId = ParticipantId;
@@ -1283,8 +1283,8 @@ void FEOSVoiceChatUser::ApplyPlayerReceivingOptions(const FGlobalParticipant& Gl
 
 	{
 		EOS_RTCAudio_UpdateParticipantVolumeOptions UpdateParticipantVolumeOptions = {};
-		UpdateParticipantVolumeOptions.ApiVersion = EOS_RTCAUDIO_UPDATEPARTICIPANTVOLUME_API_LATEST;
-		static_assert(EOS_RTCAUDIO_UPDATEPARTICIPANTVOLUME_API_LATEST == 1, "EOS_RTCAudio_UpdateParticipantVolumeOptions updated, check new fields");
+		UpdateParticipantVolumeOptions.ApiVersion = 1;
+		UE_EOS_CHECK_API_MISMATCH(EOS_RTCAUDIO_UPDATEPARTICIPANTVOLUME_API_LATEST, 1);
 		UpdateParticipantVolumeOptions.LocalUserId = LoginSession.LocalUserProductUserId;
 		UpdateParticipantVolumeOptions.RoomName = Utf8RoomName.Get();
 		UpdateParticipantVolumeOptions.ParticipantId = ParticipantId;
@@ -1318,8 +1318,8 @@ void FEOSVoiceChatUser::ApplySendingOptions(FChannelSession& ChannelSession)
 	const FTCHARToUTF8 Utf8RoomName(*ChannelSession.ChannelName);
 		
 	EOS_RTCAudio_UpdateSendingOptions UpdateSendingOptions = {};
-	UpdateSendingOptions.ApiVersion = EOS_RTCAUDIO_UPDATESENDING_API_LATEST;
-	static_assert(EOS_RTCAUDIO_UPDATESENDING_API_LATEST == 1, "EOS_RTCAudio_UpdateSendingOptions updated, check new fields");
+	UpdateSendingOptions.ApiVersion = 1;
+	UE_EOS_CHECK_API_MISMATCH(EOS_RTCAUDIO_UPDATESENDING_API_LATEST, 1);
 	UpdateSendingOptions.LocalUserId = LoginSession.LocalUserProductUserId;
 	UpdateSendingOptions.RoomName = Utf8RoomName.Get();
 	UpdateSendingOptions.AudioStatus = ChannelSession.bIsNotListening
@@ -1338,8 +1338,8 @@ void FEOSVoiceChatUser::BindLoginCallbacks()
 	// OnLobbyChannelConnectionChanged
 	{
 		EOS_Lobby_AddNotifyRTCRoomConnectionChangedOptions Options = {};
-		Options.ApiVersion = EOS_LOBBY_ADDNOTIFYRTCROOMCONNECTIONCHANGED_API_LATEST;
-		static_assert(EOS_LOBBY_ADDNOTIFYRTCROOMCONNECTIONCHANGED_API_LATEST == 2, "EOS_Lobby_AddNotifyRTCRoomConnectionChangedOptions updated, check new fields");
+		Options.ApiVersion = 2;
+		UE_EOS_CHECK_API_MISMATCH(EOS_LOBBY_ADDNOTIFYRTCROOMCONNECTIONCHANGED_API_LATEST, 2);
 		LoginSession.OnLobbyChannelConnectionChangedNotificationId = EOS_Lobby_AddNotifyRTCRoomConnectionChanged(GetLobbyInterface(), &Options, this, &FEOSVoiceChatUser::OnLobbyChannelConnectionChangedStatic);
 		if (LoginSession.OnLobbyChannelConnectionChangedNotificationId == EOS_INVALID_NOTIFICATIONID)
 		{
@@ -1368,8 +1368,8 @@ void FEOSVoiceChatUser::BindChannelCallbacks(FChannelSession& ChannelSession)
 	if(!ChannelSession.IsLobbySession())
 	{
 		EOS_RTC_AddNotifyDisconnectedOptions DisconnectedOptions = {};
-		DisconnectedOptions.ApiVersion = EOS_RTC_ADDNOTIFYDISCONNECTED_API_LATEST;
-		static_assert(EOS_RTC_ADDNOTIFYDISCONNECTED_API_LATEST == 1, "EOS_RTC_AddNotifyDisconnectedOptions updated, check new fields");
+		DisconnectedOptions.ApiVersion = 1;
+		UE_EOS_CHECK_API_MISMATCH(EOS_RTC_ADDNOTIFYDISCONNECTED_API_LATEST, 1);
 		DisconnectedOptions.LocalUserId = LoginSession.LocalUserProductUserId;
 		DisconnectedOptions.RoomName = Utf8RoomName.Get();
 		ChannelSession.OnChannelDisconnectedNotificationId = EOS_RTC_AddNotifyDisconnected(GetRtcInterface(), &DisconnectedOptions, this, &FEOSVoiceChatUser::OnChannelDisconnectedStatic);
@@ -1382,8 +1382,8 @@ void FEOSVoiceChatUser::BindChannelCallbacks(FChannelSession& ChannelSession)
 	// OnParticipantStatusChanged
 	{
 		EOS_RTC_AddNotifyParticipantStatusChangedOptions ParticipantStatusChangedOptions = {};
-		ParticipantStatusChangedOptions.ApiVersion = EOS_RTC_ADDNOTIFYPARTICIPANTSTATUSCHANGED_API_LATEST;
-		static_assert(EOS_RTC_ADDNOTIFYPARTICIPANTSTATUSCHANGED_API_LATEST == 1, "EOS_RTC_AddNotifyParticipantStatusChangedOptions updated, check new fields");
+		ParticipantStatusChangedOptions.ApiVersion = 1;
+		UE_EOS_CHECK_API_MISMATCH(EOS_RTC_ADDNOTIFYPARTICIPANTSTATUSCHANGED_API_LATEST, 1);
 		ParticipantStatusChangedOptions.LocalUserId = LoginSession.LocalUserProductUserId;
 		ParticipantStatusChangedOptions.RoomName = Utf8RoomName.Get();
 		ChannelSession.OnParticipantStatusChangedNotificationId = EOS_RTC_AddNotifyParticipantStatusChanged(GetRtcInterface(), &ParticipantStatusChangedOptions, this, &FEOSVoiceChatUser::OnChannelParticipantStatusChangedStatic);
@@ -1396,8 +1396,8 @@ void FEOSVoiceChatUser::BindChannelCallbacks(FChannelSession& ChannelSession)
 	// OnParticipantAudioUpdated
 	{
 		EOS_RTCAudio_AddNotifyParticipantUpdatedOptions ParticipantAudioUpdatedOptions = {};
-		ParticipantAudioUpdatedOptions.ApiVersion = EOS_RTCAUDIO_ADDNOTIFYPARTICIPANTUPDATED_API_LATEST;
-		static_assert(EOS_RTCAUDIO_ADDNOTIFYPARTICIPANTUPDATED_API_LATEST == 1, "EOS_RTCAudio_AddNotifyParticipantUpdatedOptions updated, check new fields");
+		ParticipantAudioUpdatedOptions.ApiVersion = 1;
+		UE_EOS_CHECK_API_MISMATCH(EOS_RTCAUDIO_ADDNOTIFYPARTICIPANTUPDATED_API_LATEST, 1);
 		ParticipantAudioUpdatedOptions.LocalUserId = LoginSession.LocalUserProductUserId;
 		ParticipantAudioUpdatedOptions.RoomName = Utf8RoomName.Get();
 		ChannelSession.OnParticipantAudioUpdatedNotificationId = EOS_RTCAudio_AddNotifyParticipantUpdated(EOS_RTC_GetAudioInterface(GetRtcInterface()), &ParticipantAudioUpdatedOptions, this, &FEOSVoiceChatUser::OnChannelParticipantAudioUpdatedStatic);
@@ -1410,8 +1410,8 @@ void FEOSVoiceChatUser::BindChannelCallbacks(FChannelSession& ChannelSession)
 	// OnAudioBeforeSend
 	{
 		EOS_RTCAudio_AddNotifyAudioBeforeSendOptions AudioBeforeSendOptions = {};
-		AudioBeforeSendOptions.ApiVersion = EOS_RTCAUDIO_ADDNOTIFYAUDIOBEFORESEND_API_LATEST;
-		static_assert(EOS_RTCAUDIO_ADDNOTIFYAUDIOBEFORESEND_API_LATEST == 1, "EOS_RTC_AddNotifyAudioBeforeSendOptions updated, check new fields");
+		AudioBeforeSendOptions.ApiVersion = 1;
+		UE_EOS_CHECK_API_MISMATCH(EOS_RTCAUDIO_ADDNOTIFYAUDIOBEFORESEND_API_LATEST, 1);
 		AudioBeforeSendOptions.LocalUserId = LoginSession.LocalUserProductUserId;
 		AudioBeforeSendOptions.RoomName = Utf8RoomName.Get();
 
@@ -1434,8 +1434,8 @@ void FEOSVoiceChatUser::BindChannelCallbacks(FChannelSession& ChannelSession)
 	// OnAudioBeforeRender
 	{
 		EOS_RTCAudio_AddNotifyAudioBeforeRenderOptions AudioBeforeRenderOptions = {};
-		AudioBeforeRenderOptions.ApiVersion = EOS_RTCAUDIO_ADDNOTIFYAUDIOBEFORERENDER_API_LATEST;
-		static_assert(EOS_RTCAUDIO_ADDNOTIFYAUDIOBEFORERENDER_API_LATEST == 1, "EOS_RTC_AddNotifyAudioBeforeRenderOptions updated, check new fields");
+		AudioBeforeRenderOptions.ApiVersion = 1;
+		UE_EOS_CHECK_API_MISMATCH(EOS_RTCAUDIO_ADDNOTIFYAUDIOBEFORERENDER_API_LATEST, 1);
 		AudioBeforeRenderOptions.LocalUserId = LoginSession.LocalUserProductUserId;
 		AudioBeforeRenderOptions.RoomName = Utf8RoomName.Get();
 		AudioBeforeRenderOptions.bUnmixedAudio = false;
@@ -1449,8 +1449,8 @@ void FEOSVoiceChatUser::BindChannelCallbacks(FChannelSession& ChannelSession)
 	// OnAudioInputState
 	{
 		EOS_RTCAudio_AddNotifyAudioInputStateOptions AudioInputStateOptions = {};
-		AudioInputStateOptions.ApiVersion = EOS_RTCAUDIO_ADDNOTIFYAUDIOINPUTSTATE_API_LATEST;
-		static_assert(EOS_RTCAUDIO_ADDNOTIFYAUDIOINPUTSTATE_API_LATEST == 1, "EOS_RTCAudio_AddNotifyAudioInputStateOptions updated, check new fields");
+		AudioInputStateOptions.ApiVersion = 1;
+		UE_EOS_CHECK_API_MISMATCH(EOS_RTCAUDIO_ADDNOTIFYAUDIOINPUTSTATE_API_LATEST, 1);
 		AudioInputStateOptions.LocalUserId = LoginSession.LocalUserProductUserId;
 		AudioInputStateOptions.RoomName = Utf8RoomName.Get();
 		ChannelSession.OnAudioInputStateNotificationId = EOS_RTCAudio_AddNotifyAudioInputState(EOS_RTC_GetAudioInterface(GetRtcInterface()), &AudioInputStateOptions, this, &FEOSVoiceChatUser::OnChannelAudioInputStateStatic);
@@ -1524,8 +1524,8 @@ void FEOSVoiceChatUser::LeaveChannelInternal(const FString& ChannelName, const F
 	const FTCHARToUTF8 Utf8RoomName(*ChannelName);
 
 	EOS_RTC_LeaveRoomOptions LeaveOptions = {};
-	LeaveOptions.ApiVersion = EOS_RTC_LEAVEROOM_API_LATEST;
-	static_assert(EOS_RTC_LEAVEROOM_API_LATEST == 1, "EOS_RTC_LeaveRoomOptions updated, check new fields");
+	LeaveOptions.ApiVersion = 1;
+	UE_EOS_CHECK_API_MISMATCH(EOS_RTC_LEAVEROOM_API_LATEST, 1);
 	LeaveOptions.LocalUserId = LoginSession.LocalUserProductUserId;
 	LeaveOptions.RoomName = Utf8RoomName.Get();
 
diff --git a/Engine/Plugins/Online/VoiceChat/EOSVoiceChat/Source/EOSVoiceChat/Private/Windows/WindowsEOSVoiceChat.cpp b/Engine/Plugins/Online/VoiceChat/EOSVoiceChat/Source/EOSVoiceChat/Private/Windows/WindowsEOSVoiceChat.cpp
--- a/Engine/Plugins/Online/VoiceChat/EOSVoiceChat/Source/EOSVoiceChat/Private/Windows/WindowsEOSVoiceChat.cpp
+++ b/Engine/Plugins/Online/VoiceChat/EOSVoiceChat/Source/EOSVoiceChat/Private/Windows/WindowsEOSVoiceChat.cpp
@@ -20,8 +20,8 @@ IEOSPlatformHandlePtr FWindowsEOSVoiceChat::EOSPlatformCreate(EOS_Platform_Optio
 	const FTCHARToUTF8 Utf8XAudioPath(*XAudioPath);
 
 	EOS_Windows_RTCOptions WindowsRTCOptions = { 0 };
-	WindowsRTCOptions.ApiVersion = EOS_WINDOWS_RTCOPTIONS_API_LATEST;
-	static_assert(EOS_WINDOWS_RTCOPTIONS_API_LATEST == 1, "EOS_Windows_RTCOptions updated, check new fields");
+	WindowsRTCOptions.ApiVersion = 1;
+	UE_EOS_CHECK_API_MISMATCH(EOS_WINDOWS_RTCOPTIONS_API_LATEST, 1);
 	WindowsRTCOptions.XAudio29DllPath = Utf8XAudioPath.Get();
 
 	const_cast<EOS_Platform_RTCOptions*>(PlatformOptions.RTCOptions)->PlatformSpecificOptions = &WindowsRTCOptions;
