Refactor of Lobby Creation logic to destroy lobby on update step failure.
Removal of session invalidation on failed update calls.

diff --git a/Engine/Plugins/Online/OnlineSubsystemEOS/Source/OnlineSubsystemEOS/Private/OnlineSessionEOS.cpp b/Engine/Plugins/Online/OnlineSubsystemEOS/Source/OnlineSubsystemEOS/Private/OnlineSessionEOS.cpp
--- a/Engine/Plugins/Online/OnlineSubsystemEOS/Source/OnlineSubsystemEOS/Private/OnlineSessionEOS.cpp
+++ b/Engine/Plugins/Online/OnlineSubsystemEOS/Source/OnlineSubsystemEOS/Private/OnlineSessionEOS.cpp
@@ -848,7 +848,11 @@ void FOnlineSessionEOS::OnMemberStatusReceived(const EOS_LobbyId& LobbyId, const
 								Session->OwningUserName = EOSSubsystem->UserManager->GetPlayerNickname(DefaultLocalUser);;
 								Session->bHosting = true;
 
-								UpdateLobbySession(Session);
+								UpdateLobbySession(Session,
+									FOnUpdateSessionCompleteDelegate::CreateLambda([this](FName SessionName, bool bWasSuccessful)
+										{
+											TriggerOnUpdateSessionCompleteDelegates(SessionName, bWasSuccessful);
+										}));
 							}
 
 							// If we are not the new owner, the new owner will update the session and we'll receive the notification, updating ours as well
@@ -1535,7 +1539,11 @@ bool FOnlineSessionEOS::UpdateSession(FName SessionName, FOnlineSessionSettings&
 		{
 			if (Session->SessionSettings.bUseLobbiesIfAvailable)
 			{
-				Result = UpdateLobbySession(Session);
+				Result = UpdateLobbySession(Session,
+					FOnUpdateSessionCompleteDelegate::CreateLambda([this](FName SessionName, bool bWasSuccessful)
+					{
+						TriggerOnUpdateSessionCompleteDelegates(SessionName, bWasSuccessful);
+					}));
 			}
 			else
 			{
@@ -1601,8 +1609,7 @@ uint32 FOnlineSessionEOS::UpdateEOSSession(FNamedOnlineSession* Session)
 			bWasSuccessful = Data->ResultCode == EOS_EResult::EOS_Success || Data->ResultCode == EOS_EResult::EOS_Sessions_OutOfSync;
 			if (!bWasSuccessful)
 			{
-				Session->SessionState = EOnlineSessionState::NoSession;
-				UE_LOG_ONLINE_SESSION(Error, TEXT("EOS_Sessions_UpdateSession() failed with EOS result code (%s)"), ANSI_TO_TCHAR(EOS_EResult_ToString(Data->ResultCode)));
+				UE_LOG_ONLINE_SESSION(Warning, TEXT("EOS_Sessions_UpdateSession() failed with EOS result code (%s)"), ANSI_TO_TCHAR(EOS_EResult_ToString(Data->ResultCode)));
 			}
 		}
 		else
@@ -3525,7 +3532,7 @@ void FOnlineSessionEOS::RemovePlayerFromSession(int32 LocalUserNum, FName Sessio
 		{
 			if (Data->ResultCode == EOS_EResult::EOS_Success)
 			{
-				UE_LOG_ONLINE_SESSION(Verbose, TEXT("[FOnlineSessionEOS::RemovePlayerFromSession] KickMember finished successfully for lobby %d."), Data->LobbyId);
+				UE_LOG_ONLINE_SESSION(Verbose, TEXT("[FOnlineSessionEOS::RemovePlayerFromSession] KickMember finished successfully for lobby %s."), UTF8_TO_TCHAR(Data->LobbyId));
 			}
 			else
 			{
@@ -3676,7 +3683,7 @@ uint32 FOnlineSessionEOS::CreateLobbySession(int32 HostingPlayerNum, FNamedOnlin
 			bool bWasSuccessful = Data->ResultCode == EOS_EResult::EOS_Success;
 			if (bWasSuccessful)
 			{
-				UE_LOG_ONLINE_SESSION(Verbose, TEXT("[FOnlineSessionEOS::CreateLobbySession] CreateLobby was successful. LobbyId is %d."), Data->LobbyId);
+				UE_LOG_ONLINE_SESSION(Verbose, TEXT("[FOnlineSessionEOS::CreateLobbySession] CreateLobby was successful. LobbyId is %s."), UTF8_TO_TCHAR(Data->LobbyId));
 
 				Session->SessionState = EOnlineSessionState::Pending;
 
@@ -3685,7 +3692,7 @@ uint32 FOnlineSessionEOS::CreateLobbySession(int32 HostingPlayerNum, FNamedOnlin
 				FInternetAddrEOS TempAddr(LexToString(LocalProductUserId), SessionName.ToString(), FURL::UrlConfig.DefaultPort);
 				FString HostAddr = TempAddr.ToString(true);
 
-				Session->SessionInfo = MakeShareable(new FOnlineSessionInfoEOS(HostAddr, FUniqueNetIdEOSLobby::Create(Data->LobbyId), nullptr));
+				Session->SessionInfo = MakeShareable(new FOnlineSessionInfoEOS(HostAddr, FUniqueNetIdEOSLobby::Create(UTF8_TO_TCHAR(Data->LobbyId)), nullptr));
 
 #if WITH_EOS_RTC
 				if (FEOSVoiceChatUser* VoiceChatUser = static_cast<FEOSVoiceChatUser*>(EOSSubsystem->GetEOSVoiceChatUserInterface(*LocalUserNetId)))
@@ -3696,7 +3703,7 @@ uint32 FOnlineSessionEOS::CreateLobbySession(int32 HostingPlayerNum, FNamedOnlin
 
 				BeginSessionAnalytics(Session);
 
-				UpdateLobbySession(Session);
+				UpdateLobbySession(Session, FOnUpdateSessionCompleteDelegate::CreateThreadSafeSP(this, &FOnlineSessionEOS::OnCreateLobbySessionUpdateComplete));
 			}
 			else
 			{
@@ -3705,9 +3712,9 @@ uint32 FOnlineSessionEOS::CreateLobbySession(int32 HostingPlayerNum, FNamedOnlin
 				Session->SessionState = EOnlineSessionState::NoSession;
 
 				RemoveNamedSession(SessionName);
-			}
 
-			TriggerOnCreateSessionCompleteDelegates(SessionName, bWasSuccessful);
+				TriggerOnCreateSessionCompleteDelegates(SessionName, bWasSuccessful);
+			}
 		}
 	};
 
@@ -3716,6 +3723,82 @@ uint32 FOnlineSessionEOS::CreateLobbySession(int32 HostingPlayerNum, FNamedOnlin
 	return ONLINE_IO_PENDING;
 }
 
+void FOnlineSessionEOS::OnCreateLobbySessionUpdateComplete(FName SessionName, bool bWasSuccessful)
+{
+	if (!bWasSuccessful)
+	{
+		UE_LOG_ONLINE_SESSION(Warning, TEXT("[OnCreateLobbySessionUpdateComplete] UpdateLobby not successful. Created session [%s] will be destroyed."), *SessionName.ToString());
+
+		FNamedOnlineSession* Session = GetNamedSession(SessionName);
+		if (Session)
+		{
+			DestroyLobbySessionOnCreationUpdateError(Session);
+		}
+	}
+	else
+	{
+		TriggerOnCreateSessionCompleteDelegates(SessionName, bWasSuccessful);
+	}
+}
+
+void FOnlineSessionEOS::DestroyLobbySessionOnCreationUpdateError(FNamedOnlineSession* Session)
+{
+	check(Session != nullptr);
+	check(Session->SessionInfo.IsValid());
+
+	FOnlineSessionInfoEOS* SessionInfo = (FOnlineSessionInfoEOS*)(Session->SessionInfo.Get());
+	check(Session->SessionSettings.bUseLobbiesIfAvailable); // We check if it's a lobby session
+
+	EOS_Lobby_DestroyLobbyOptions DestroyOptions = { 0 };
+	DestroyOptions.ApiVersion = 1;
+	UE_EOS_CHECK_API_MISMATCH(EOS_LOBBY_DESTROYLOBBY_API_LATEST, 1);
+	const FTCHARToUTF8 Utf8LobbyId(*SessionInfo->GetSessionId().ToString());
+	DestroyOptions.LobbyId = (EOS_LobbyId)Utf8LobbyId.Get();
+	DestroyOptions.LocalUserId = EOSSubsystem->UserManager->GetLocalProductUserId(EOSSubsystem->UserManager->GetDefaultLocalUser()); // Maybe not split screen friendly
+
+	FName SessionName = Session->SessionName;
+	FLobbyDestroyedCallback* DestroyCallbackObj = new FLobbyDestroyedCallback(FOnlineSessionEOSWeakPtr(AsShared()));
+	LobbyLeftCallback = DestroyCallbackObj;
+	DestroyCallbackObj->CallbackLambda = [this, SessionName](const EOS_Lobby_DestroyLobbyCallbackInfo* Data)
+	{
+		FNamedOnlineSession* LobbySession = GetNamedSession(SessionName);
+		if (LobbySession)
+		{
+			bool bWasSuccessful = Data->ResultCode == EOS_EResult::EOS_Success;
+			if (bWasSuccessful)
+			{
+				UE_LOG_ONLINE_SESSION(Verbose, TEXT("[FOnlineSessionEOS::DestroyLobbySessionOnCreationUpdateError] DestroyLobby was successful. LobbyId is %s."), UTF8_TO_TCHAR(Data->LobbyId));
+			}
+			else
+			{
+				UE_LOG_ONLINE_SESSION(Warning, TEXT("[FOnlineSessionEOS::DestroyLobbySessionOnCreationUpdateError] DestroyLobby not successful. Finished with EOS_EResult %s"), ANSI_TO_TCHAR(EOS_EResult_ToString(Data->ResultCode)));
+			}
+
+#if WITH_EOS_RTC
+			if (FEOSVoiceChatUser* VoiceChatUser = static_cast<FEOSVoiceChatUser*>(EOSSubsystem->GetEOSVoiceChatUserInterface(*EOSSubsystem->UserManager->GetLocalUniqueNetIdEOS())))
+			{
+				VoiceChatUser->RemoveLobbyRoom(UTF8_TO_TCHAR(Data->LobbyId));
+			}
+#endif
+
+			EndSessionAnalytics();
+
+			LobbySession->SessionState = EOnlineSessionState::NoSession;
+
+			RemoveNamedSession(SessionName);
+		}
+		else
+		{
+			UE_LOG_ONLINE_SESSION(Warning, TEXT("[FOnlineSessionEOS::DestroyLobbySession] Unable to find session %s"), *SessionName.ToString());
+		}
+
+		// Whether Lobby destruction is successful or not, creation will always be failed at this point
+		TriggerOnCreateSessionCompleteDelegates(SessionName, false);
+	};
+
+	EOS_Lobby_DestroyLobby(LobbyHandle, &DestroyOptions, DestroyCallbackObj, DestroyCallbackObj->GetCallbackPtr());
+}
+
 uint32 FOnlineSessionEOS::JoinLobbySession(int32 PlayerNum, FNamedOnlineSession* Session, const FOnlineSession* SearchSession)
 {
 	check(Session != nullptr);
@@ -3759,7 +3842,7 @@ uint32 FOnlineSessionEOS::JoinLobbySession(int32 PlayerNum, FNamedOnlineSession*
 					const bool bWasSuccessful = Data->ResultCode == EOS_EResult::EOS_Success;
 					if (bWasSuccessful)
 					{
-						UE_LOG_ONLINE_SESSION(Verbose, TEXT("[FOnlineSessionEOS::JoinLobbySession] JoinLobby was successful. LobbyId is %d."), Data->LobbyId);
+						UE_LOG_ONLINE_SESSION(Verbose, TEXT("[FOnlineSessionEOS::JoinLobbySession] JoinLobby was successful. LobbyId is %s."), UTF8_TO_TCHAR(Data->LobbyId));
 
 						BeginSessionAnalytics(Session);
 
@@ -3946,7 +4029,7 @@ void FOnlineSessionEOS::SetLobbyAttributes(EOS_HLobbyModification LobbyModificat
 	}
 }
 
-uint32 FOnlineSessionEOS::UpdateLobbySession(FNamedOnlineSession* Session)
+uint32 FOnlineSessionEOS::UpdateLobbySession(FNamedOnlineSession* Session, const FOnUpdateSessionCompleteDelegate& CompletionDelegate)
 {
 	check(Session != nullptr);
 
@@ -3955,6 +4038,13 @@ uint32 FOnlineSessionEOS::UpdateLobbySession(FNamedOnlineSession* Session)
 	if (Session->SessionState == EOnlineSessionState::Creating)
 	{
 		Result = ONLINE_IO_PENDING;
+
+		UE_LOG_ONLINE_SESSION(Warning, TEXT("[FOnlineSessionEOS::UpdateLobbySession] UpdateLobby not successful. Session %s is still being created."), *Session->SessionName.ToString());
+
+		EOSSubsystem->ExecuteNextTick([CompletionDelegate, SessionName = Session->SessionName]()
+			{
+				CompletionDelegate.ExecuteIfBound(SessionName, false);
+			});
 	}
 	else
 	{
@@ -3981,7 +4071,7 @@ uint32 FOnlineSessionEOS::UpdateLobbySession(FNamedOnlineSession* Session)
 
 			FName SessionName = Session->SessionName;
 			FLobbyUpdatedCallback* CallbackObj = new FLobbyUpdatedCallback(FOnlineSessionEOSWeakPtr(AsShared()));
-			CallbackObj->CallbackLambda = [this, SessionName](const EOS_Lobby_UpdateLobbyCallbackInfo* Data)
+			CallbackObj->CallbackLambda = [this, SessionName, CompletionDelegate](const EOS_Lobby_UpdateLobbyCallbackInfo* Data)
 			{
 				FNamedOnlineSession* Session = GetNamedSession(SessionName);
 				if (Session)
@@ -3989,16 +4079,15 @@ uint32 FOnlineSessionEOS::UpdateLobbySession(FNamedOnlineSession* Session)
 					bool bWasSuccessful = Data->ResultCode == EOS_EResult::EOS_Success || Data->ResultCode == EOS_EResult::EOS_Sessions_OutOfSync;
 					if (!bWasSuccessful)
 					{
-						Session->SessionState = EOnlineSessionState::NoSession;
 						UE_LOG_ONLINE_SESSION(Warning, TEXT("[FOnlineSessionEOS::UpdateLobbySession] UpdateLobby not successful. Finished with EOS_EResult %s"), ANSI_TO_TCHAR(EOS_EResult_ToString(Data->ResultCode)));
 					}
 
-					TriggerOnUpdateSessionCompleteDelegates(SessionName, bWasSuccessful);
+					CompletionDelegate.ExecuteIfBound(SessionName, bWasSuccessful);
 				}
 				else
 				{
 					UE_LOG_ONLINE_SESSION(Warning, TEXT("[FOnlineSessionEOS::UpdateLobbySession] Unable to find session %s"), *SessionName.ToString());
-					TriggerOnUpdateSessionCompleteDelegates(SessionName, false);
+					CompletionDelegate.ExecuteIfBound(SessionName, false);
 				}
 			};
 
@@ -4011,6 +4100,11 @@ uint32 FOnlineSessionEOS::UpdateLobbySession(FNamedOnlineSession* Session)
 		else
 		{
 			UE_LOG_ONLINE_SESSION(Warning, TEXT("[FOnlineSessionEOS::UpdateLobbySession] UpdateLobbyModification not successful. Finished with EOS_EResult %s"), ANSI_TO_TCHAR(EOS_EResult_ToString(LobbyModificationResult)));
+
+			EOSSubsystem->ExecuteNextTick([CompletionDelegate, SessionName = Session->SessionName]()
+				{
+					CompletionDelegate.ExecuteIfBound(SessionName, false);
+				});
 		}
 	}
 
@@ -4071,7 +4165,7 @@ uint32 FOnlineSessionEOS::DestroyLobbySession(FNamedOnlineSession* Session, cons
 				bool bWasSuccessful = Data->ResultCode == EOS_EResult::EOS_Success;
 				if (bWasSuccessful)
 				{
-					UE_LOG_ONLINE_SESSION(Verbose, TEXT("[FOnlineSessionEOS::DestroyLobbySession] LeaveLobby was successful. LobbyId is %s."), Data->LobbyId);
+					UE_LOG_ONLINE_SESSION(Verbose, TEXT("[FOnlineSessionEOS::DestroyLobbySession] LeaveLobby was successful. LobbyId is %s."), UTF8_TO_TCHAR(Data->LobbyId));
 				}
 				else
 				{
@@ -4268,10 +4362,10 @@ void FOnlineSessionEOS::AddLobbySearchResult(const TSharedRef<FLobbyDetailsEOS>&
 		FInternetAddrEOS TempAddr(LexToString(LobbyDetailsInfo->LobbyOwnerUserId), NetDriverName.ToString(), GetTypeHash(NetDriverName.ToString()));
 		FString HostAddr = TempAddr.ToString(true);
 
-		SearchResult.Session.SessionInfo = MakeShareable(new FOnlineSessionInfoEOS(HostAddr, FUniqueNetIdEOSLobby::Create(LobbyDetailsInfo->LobbyId), nullptr));
+		SearchResult.Session.SessionInfo = MakeShareable(new FOnlineSessionInfoEOS(HostAddr, FUniqueNetIdEOSLobby::Create(UTF8_TO_TCHAR(LobbyDetailsInfo->LobbyId)), nullptr));
 
 		// We copy the lobby data and settings
-		LobbySearchResultsCache.Add(FString(LobbyDetailsInfo->LobbyId), LobbyDetails);
+		LobbySearchResultsCache.Add(FString(UTF8_TO_TCHAR(LobbyDetailsInfo->LobbyId)), LobbyDetails);
 		CopyLobbyData(LobbyDetails, LobbyDetailsInfo, SearchResult.Session, Callback);
 
 		EOS_LobbyDetails_Info_Release(LobbyDetailsInfo);
@@ -4343,7 +4437,7 @@ void FOnlineSessionEOS::CopyLobbyData(const TSharedRef<FLobbyDetailsEOS>& LobbyD
 
 	if (!TargetUserIds.IsEmpty())
 	{
-		EOSSubsystem->UserManager->ResolveUniqueNetIds(TargetUserIds, [this, LobbyDetails, LobbyId = FUniqueNetIdEOSLobby::Create(LobbyDetailsInfo->LobbyId), OriginalCallback = Callback](TMap<EOS_ProductUserId, FUniqueNetIdEOSRef> ResolvedUniqueNetIds)
+		EOSSubsystem->UserManager->ResolveUniqueNetIds(TargetUserIds, [this, LobbyDetails, LobbyId = FUniqueNetIdEOSLobby::Create(UTF8_TO_TCHAR(LobbyDetailsInfo->LobbyId)), OriginalCallback = Callback](TMap<EOS_ProductUserId, FUniqueNetIdEOSRef> ResolvedUniqueNetIds)
 			{
 				
 				FOnlineSession* Session = GetOnlineSessionFromLobbyId(*LobbyId);
diff --git a/Engine/Plugins/Online/OnlineSubsystemEOS/Source/OnlineSubsystemEOS/Private/OnlineSessionEOS.h b/Engine/Plugins/Online/OnlineSubsystemEOS/Source/OnlineSubsystemEOS/Private/OnlineSessionEOS.h
--- a/Engine/Plugins/Online/OnlineSubsystemEOS/Source/OnlineSubsystemEOS/Private/OnlineSessionEOS.h
+++ b/Engine/Plugins/Online/OnlineSubsystemEOS/Source/OnlineSubsystemEOS/Private/OnlineSessionEOS.h
@@ -219,10 +219,12 @@ private:
 	FCallbackBase* LobbySendInviteCallback;
 
 	uint32 CreateLobbySession(int32 HostingPlayerNum, FNamedOnlineSession* Session);
+	void OnCreateLobbySessionUpdateComplete(FName SessionName, bool bWasSuccessful);
+	void DestroyLobbySessionOnCreationUpdateError(FNamedOnlineSession* Session);
 	uint32 FindLobbySession(int32 SearchingPlayerNum, const TSharedRef<FOnlineSessionSearch>& SearchSettings);
 	void StartLobbySearch(int32 SearchingPlayerNum, EOS_HLobbySearch LobbySearchHandle, const TSharedRef<FOnlineSessionSearch>& SearchSettings, const FOnSingleSessionResultCompleteDelegate& CompletionDelegate);
 	uint32 JoinLobbySession(int32 PlayerNum, FNamedOnlineSession* Session, const FOnlineSession* SearchSession);
-	uint32 UpdateLobbySession(FNamedOnlineSession* Session);
+	uint32 UpdateLobbySession(FNamedOnlineSession* Session, const FOnUpdateSessionCompleteDelegate& CompletionDelegate);
 	uint32 StartLobbySession(FNamedOnlineSession* Session);
 	uint32 EndLobbySession(FNamedOnlineSession* Session);
 	uint32 DestroyLobbySession(FNamedOnlineSession* Session, const FOnDestroySessionCompleteDelegate& CompletionDelegate);
