- Fix ParseArtifactSettingsFromConfigString attempting to read only from ClientEncryptionKey. This was an oversight. We have a config redirect from EncryptionKey to ClientEncryptionKey, but this obviously only works for _keys_, but in this case EncryptionKey is part of the Artifact value.

diff --git a/Engine/Plugins/Online/OnlineSubsystemEOS/Source/OnlineSubsystemEOS/Private/EOSSettings.cpp b/Engine/Plugins/Online/OnlineSubsystemEOS/Source/OnlineSubsystemEOS/Private/EOSSettings.cpp
--- a/Engine/Plugins/Online/OnlineSubsystemEOS/Source/OnlineSubsystemEOS/Private/EOSSettings.cpp
+++ b/Engine/Plugins/Online/OnlineSubsystemEOS/Source/OnlineSubsystemEOS/Private/EOSSettings.cpp
@@ -118,7 +118,9 @@ FEOSArtifactSettings ParseArtifactSettingsFromConfigString(const FString& RawLin
 		{
 			Result.DeploymentId = StripQuotes(Values[ValueIndex + 1]);
 		}
-		else if (Values[ValueIndex] == TEXT("ClientEncryptionKey"))
+		// EncryptionKey is problematic as a key name as it gets removed by IniKeyDenyList, so lots of EOS config has moved to ClientEncryptionKey instead.
+		// That specific issue doesn't affect this specific case as it's part of a config _value_, but supporting both names for consistency and back-compat.
+		else if (Values[ValueIndex] == TEXT("EncryptionKey") || Values[ValueIndex] == TEXT("ClientEncryptionKey"))
 		{
 			Result.EncryptionKey = StripQuotes(Values[ValueIndex + 1]);
 		}
