- Fix LoginViaExternalAuth not using the config auth scopes.

diff --git a/Engine/Plugins/Online/OnlineSubsystemEOS/Source/OnlineSubsystemEOS/Private/UserManagerEOS.cpp b/Engine/Plugins/Online/OnlineSubsystemEOS/Source/OnlineSubsystemEOS/Private/UserManagerEOS.cpp
--- a/Engine/Plugins/Online/OnlineSubsystemEOS/Source/OnlineSubsystemEOS/Private/UserManagerEOS.cpp
+++ b/Engine/Plugins/Online/OnlineSubsystemEOS/Source/OnlineSubsystemEOS/Private/UserManagerEOS.cpp
@@ -129,6 +129,27 @@ static inline EOS_EExternalCredentialType ToEOS_EExternalCredentialType(FName OS
 	return EOS_EExternalCredentialType::EOS_ECT_OPENID_ACCESS_TOKEN;
 }
 
+namespace {
+
+EOS_EAuthScopeFlags GetAuthScopeFlags()
+{
+	const FEOSSettings& Settings = UEOSSettings::GetSettings();
+
+	EOS_EAuthScopeFlags ScopeFlags = EOS_EAuthScopeFlags::EOS_AS_NoFlags;
+	for (const FString& FlagsStr : Settings.AuthScopeFlags)
+	{
+		EOS_EAuthScopeFlags Flags;
+		if (LexFromString(Flags, FlagsStr))
+		{
+			ScopeFlags |= Flags;
+		}
+	}
+	
+	return ScopeFlags;
+}
+
+} // namespace
+
 /** Delegates that are used for internal calls and are meant to be ignored */
 IOnlinePresence::FOnPresenceTaskCompleteDelegate IgnoredPresenceDelegate;
 IOnlineUser::FOnQueryExternalIdMappingsComplete IgnoredMappingDelegate;
@@ -444,15 +465,7 @@ bool FUserManagerEOS::Login(int32 LocalUserNum, const FOnlineAccountCredentials&
 	LoginOptions.ApiVersion = 2;
 	UE_EOS_CHECK_API_MISMATCH(EOS_AUTH_LOGIN_API_LATEST, 2);
 
-	LoginOptions.ScopeFlags = EOS_EAuthScopeFlags::EOS_AS_NoFlags;
-	for (const FString& FlagsStr : Settings.AuthScopeFlags)
-	{
-		EOS_EAuthScopeFlags Flags;
-		if (LexFromString(Flags, FlagsStr))
-		{
-			LoginOptions.ScopeFlags |= Flags;
-		}
-	}
+	LoginOptions.ScopeFlags = GetAuthScopeFlags();
 
 	FPlatformEOSHelpersPtr EOSHelpers = EOSSubsystem->GetEOSHelpers();
 
@@ -573,7 +586,8 @@ void FUserManagerEOS::LoginViaExternalAuth(int32 LocalUserNum)
 				EOS_Auth_LoginOptions LoginOptions = { };
 				LoginOptions.ApiVersion = 2;
 				UE_EOS_CHECK_API_MISMATCH(EOS_AUTH_LOGIN_API_LATEST, 2);
-				LoginOptions.ScopeFlags = EOS_EAuthScopeFlags::EOS_AS_BasicProfile | EOS_EAuthScopeFlags::EOS_AS_FriendsList | EOS_EAuthScopeFlags::EOS_AS_Presence;
+
+				LoginOptions.ScopeFlags = GetAuthScopeFlags();
 
 				check(LocalUserNumToLastLoginCredentials.Contains(LocalUserNum));
 				FAuthCredentials Credentials(ToEOS_EExternalCredentialType(GetPlatformOSS()->GetSubsystemName(), *LocalUserNumToLastLoginCredentials[LocalUserNum]), AuthToken);
