- Fix EncryptionKey being set to 0 in some projects, causing EOSSDK to throw warnings.
- Fix OSSEOS not setting EncryptionKey to nullptr when one is not specified in config.
- Fix OSSEOS failing init if title/player data storage are not available due to no encryption key

diff --git a/Engine/Plugins/Online/OnlineSubsystemEOS/Source/OnlineSubsystemEOS/Private/OnlineSubsystemEOS.cpp b/Engine/Plugins/Online/OnlineSubsystemEOS/Source/OnlineSubsystemEOS/Private/OnlineSubsystemEOS.cpp
--- a/Engine/Plugins/Online/OnlineSubsystemEOS/Source/OnlineSubsystemEOS/Private/OnlineSubsystemEOS.cpp
+++ b/Engine/Plugins/Online/OnlineSubsystemEOS/Source/OnlineSubsystemEOS/Private/OnlineSubsystemEOS.cpp
@@ -124,33 +124,6 @@ public:
 
 #endif // WITH_EOS_RTC
 
-/** Class that holds the strings for the call duration */
-struct FEOSPlatformOptions :
-	public EOS_Platform_Options
-{
-	FEOSPlatformOptions() :
-		EOS_Platform_Options()
-	{
-		ApiVersion = 12;
-		UE_EOS_CHECK_API_MISMATCH(EOS_PLATFORM_OPTIONS_API_LATEST, 12);
-		ProductId = ProductIdAnsi;
-		SandboxId = SandboxIdAnsi;
-		DeploymentId = DeploymentIdAnsi;
-		ClientCredentials.ClientId = ClientIdAnsi;
-		ClientCredentials.ClientSecret = ClientSecretAnsi;
-		CacheDirectory = CacheDirectoryAnsi;
-		EncryptionKey = EncryptionKeyAnsi;
-	}
-
-	char ClientIdAnsi[EOS_OSS_STRING_BUFFER_LENGTH];
-	char ClientSecretAnsi[EOS_OSS_STRING_BUFFER_LENGTH];
-	char ProductIdAnsi[EOS_OSS_STRING_BUFFER_LENGTH];
-	char SandboxIdAnsi[EOS_OSS_STRING_BUFFER_LENGTH];
-	char DeploymentIdAnsi[EOS_OSS_STRING_BUFFER_LENGTH];
-	char CacheDirectoryAnsi[EOS_OSS_STRING_BUFFER_LENGTH];
-	char EncryptionKeyAnsi[EOS_ENCRYPTION_KEY_MAX_BUFFER_LEN];
-};
-
 FPlatformEOSHelpersPtr FOnlineSubsystemEOS::EOSHelpersPtr;
 
 void FOnlineSubsystemEOS::ModuleInit()
@@ -203,15 +176,26 @@ bool FOnlineSubsystemEOS::PlatformCreate()
 		return false;
 	}
 
+	const auto ClientIdUtf8 = StringCast<UTF8CHAR>(*ArtifactSettings.ClientId);
+	const auto ClientSecretUtf8 = StringCast<UTF8CHAR>(*ArtifactSettings.ClientSecret);
+	const auto ProductIdUtf8 = StringCast<UTF8CHAR>(*ArtifactSettings.ProductId);
+	const auto SandboxIdUtf8 = StringCast<UTF8CHAR>(*ArtifactSettings.SandboxId);
+	const auto DeploymentIdUtf8 = StringCast<UTF8CHAR>(*ArtifactSettings.DeploymentId);
+	const auto EncryptionKeyUtf8 = StringCast<UTF8CHAR>(*ArtifactSettings.EncryptionKey);
+
 	// Create platform instance
-	FEOSPlatformOptions PlatformOptions;
-	FCStringAnsi::Strncpy(PlatformOptions.ClientIdAnsi, TCHAR_TO_UTF8(*ArtifactSettings.ClientId), EOS_OSS_STRING_BUFFER_LENGTH);
-	FCStringAnsi::Strncpy(PlatformOptions.ClientSecretAnsi, TCHAR_TO_UTF8(*ArtifactSettings.ClientSecret), EOS_OSS_STRING_BUFFER_LENGTH);
-	FCStringAnsi::Strncpy(PlatformOptions.ProductIdAnsi, TCHAR_TO_UTF8(*ArtifactSettings.ProductId), EOS_OSS_STRING_BUFFER_LENGTH);
-	FCStringAnsi::Strncpy(PlatformOptions.SandboxIdAnsi, TCHAR_TO_UTF8(*ArtifactSettings.SandboxId), EOS_OSS_STRING_BUFFER_LENGTH);
-	FCStringAnsi::Strncpy(PlatformOptions.DeploymentIdAnsi, TCHAR_TO_UTF8(*ArtifactSettings.DeploymentId), EOS_OSS_STRING_BUFFER_LENGTH);
+	EOS_Platform_Options PlatformOptions = {};
+	PlatformOptions.ApiVersion = 12;
+	UE_EOS_CHECK_API_MISMATCH(EOS_PLATFORM_OPTIONS_API_LATEST, 12);
+	PlatformOptions.ClientCredentials.ClientId = reinterpret_cast<const char*>(ArtifactSettings.ClientId.IsEmpty() ? nullptr : (const char*)ClientIdUtf8.Get());
+	PlatformOptions.ClientCredentials.ClientSecret = ArtifactSettings.ClientSecret.IsEmpty() ? nullptr : (const char*)ClientSecretUtf8.Get();
+	PlatformOptions.ProductId = ArtifactSettings.ProductId.IsEmpty() ? nullptr : (const char*)ProductIdUtf8.Get();
+	PlatformOptions.SandboxId = ArtifactSettings.SandboxId.IsEmpty() ? nullptr : (const char*)SandboxIdUtf8.Get();
+	PlatformOptions.DeploymentId = ArtifactSettings.DeploymentId.IsEmpty() ? nullptr : (const char*)DeploymentIdUtf8.Get();
+	PlatformOptions.EncryptionKey = ArtifactSettings.EncryptionKey.IsEmpty() ? nullptr : (const char*)EncryptionKeyUtf8.Get();
 	PlatformOptions.bIsServer = IsRunningDedicatedServer() ? EOS_TRUE : EOS_FALSE;
 	PlatformOptions.Reserved = nullptr;
+
 	FEOSSettings EOSSettings = UEOSSettings::GetSettings();
 	uint64 OverlayFlags = 0;
 	if (!EOSSettings.bEnableOverlay)
@@ -234,19 +218,16 @@ bool FOnlineSubsystemEOS::PlatformCreate()
 	const bool bOverlayAllowed = IsRunningGame() || bEditorOverlayAllowed;
 
 	PlatformOptions.Flags = bOverlayAllowed ? OverlayFlags : EOS_PF_DISABLE_OVERLAY;
-	// Make the cache directory be in the user's writable area
 
+	// Make the cache directory be in the user's writable area
+	FString CacheDir;
+	
 	if (FPlatformMisc::IsCacheStorageAvailable())
 	{
-		const FString CacheDir = EOSSDKManager->GetCacheDirBase() / ArtifactSettings.ArtifactName / EOSSettings.CacheDir;
-		FCStringAnsi::Strncpy(PlatformOptions.CacheDirectoryAnsi, TCHAR_TO_UTF8(*CacheDir), EOS_OSS_STRING_BUFFER_LENGTH);
-	}
-	else
-	{
-		PlatformOptions.CacheDirectory = nullptr;
+		CacheDir = EOSSDKManager->GetCacheDirBase() / ArtifactSettings.ArtifactName / EOSSettings.CacheDir;
 	}
-
-	FCStringAnsi::Strncpy(PlatformOptions.EncryptionKeyAnsi, TCHAR_TO_UTF8(*ArtifactSettings.EncryptionKey), EOS_ENCRYPTION_KEY_MAX_BUFFER_LEN);
+	const auto CacheDirUtf8 = StringCast<UTF8CHAR>(*CacheDir);
+	PlatformOptions.CacheDirectory = CacheDir.IsEmpty() ? nullptr : (const char*)CacheDirUtf8.Get();
 
 #if WITH_EOS_RTC
 	EOS_Platform_RTCOptions RtcOptions = { 0 };
@@ -379,14 +360,45 @@ bool FOnlineSubsystemEOS::Init()
 		}
 		StoreInterfacePtr = MakeShareable(new FOnlineStoreEOS(this));
 	}
+
+	// We set the product id
+	FEOSArtifactSettings ArtifactSettings;
+	if (UEOSSettings::GetSelectedArtifactSettings(ArtifactSettings))
+	{
+		ProductId = ArtifactSettings.ProductId;
+	}
+	else
+	{
+		// This really should not be possible, if we made it past PlatformCreate.
+		checkNoEntry();
+		UE_LOG_ONLINE(Warning, TEXT("[FOnlineSubsystemEOS::Init] GetSelectedArtifactSettings failed, ProductIdAnsi not set."));
+	}
+
 	TitleStorageHandle = EOS_Platform_GetTitleStorageInterface(*EOSPlatformHandle);
-	if (TitleStorageHandle == nullptr)
+	if (TitleStorageHandle)
+	{
+		TitleFileInterfacePtr = MakeShared<FOnlineTitleFileEOS>(this);
+	}
+	else if (ArtifactSettings.EncryptionKey.IsEmpty())
+	{
+		UE_LOG_ONLINE(Verbose, TEXT("FOnlineSubsystemEOS: ArtifactName=[%s] EncryptionKey unset, TitleFile interface unavailable."));
+	}
+	else
 	{
 		UE_LOG_ONLINE(Error, TEXT("FOnlineSubsystemEOS: failed to init EOS platform, couldn't get title storage handle"));
 		return false;
 	}
+
 	PlayerDataStorageHandle = EOS_Platform_GetPlayerDataStorageInterface(*EOSPlatformHandle);
-	if (PlayerDataStorageHandle == nullptr)
+	if (PlayerDataStorageHandle)
+	{
+		UserCloudInterfacePtr = MakeShared<FOnlineUserCloudEOS>(this);
+	}
+	else if (ArtifactSettings.EncryptionKey.IsEmpty())
+	{
+		UE_LOG_ONLINE(Verbose, TEXT("FOnlineSubsystemEOS: ArtifactName=[%s] EncryptionKey unset, UserCloud interface unavailable."));
+	}
+	else
 	{
 		UE_LOG_ONLINE(Error, TEXT("FOnlineSubsystemEOS: failed to init EOS platform, couldn't get player data storage handle"));
 		return false;
@@ -400,19 +412,6 @@ bool FOnlineSubsystemEOS::Init()
 		UE_LOG_ONLINE(Warning, TEXT("[FOnlineSubsystemEOS::Init] Unable to initialize Socket Subsystem. Error=[%s]"), *ErrorMessage);
 	}
 
-	// We set the product id
-	FEOSArtifactSettings ArtifactSettings;
-	if (UEOSSettings::GetSelectedArtifactSettings(ArtifactSettings))
-	{
-		ProductId = ArtifactSettings.ProductId;
-	}
-	else
-	{
-		// This really should not be possible, if we made it past PlatformCreate.
-		checkNoEntry();
-		UE_LOG_ONLINE(Warning, TEXT("[FOnlineSubsystemEOS::Init] GetSelectedArtifactSettings failed, ProductIdAnsi not set."));
-	}
-
 	UserManager = MakeShareable(new FUserManagerEOS(this));
 	UserManager->Init();
 	SessionInterfacePtr = MakeShareable(new FOnlineSessionEOS(this));
@@ -430,8 +429,6 @@ bool FOnlineSubsystemEOS::Init()
 	StatsInterfacePtr = MakeShareable(new FOnlineStatsEOS(this));
 	LeaderboardsInterfacePtr = MakeShareable(new FOnlineLeaderboardsEOS(this));
 	AchievementsInterfacePtr = MakeShareable(new FOnlineAchievementsEOS(this));
-	TitleFileInterfacePtr = MakeShareable(new FOnlineTitleFileEOS(this));
-	UserCloudInterfacePtr = MakeShareable(new FOnlineUserCloudEOS(this));
 
 	// We initialized ok so we can tick
 	StartTicker();
