Refactor of session update code in OSS EOS Session interface to ensure non-owner lobby members can only update their own lobby member settings (to match API functionality).

diff --git a/Engine/Plugins/Online/OnlineSubsystemEOS/Source/OnlineSubsystemEOS/Private/OnlineSessionEOS.cpp b/Engine/Plugins/Online/OnlineSubsystemEOS/Source/OnlineSubsystemEOS/Private/OnlineSessionEOS.cpp
--- a/Engine/Plugins/Online/OnlineSubsystemEOS/Source/OnlineSubsystemEOS/Private/OnlineSessionEOS.cpp
+++ b/Engine/Plugins/Online/OnlineSubsystemEOS/Source/OnlineSubsystemEOS/Private/OnlineSessionEOS.cpp
@@ -4005,28 +4005,32 @@ void FOnlineSessionEOS::SetLobbyAttributes(EOS_HLobbyModification LobbyModificat
 		AddLobbyAttribute(LobbyModificationHandle, &Attribute);
 	}
 
-	// Add all of the member settings
-	for (TPair<FUniqueNetIdRef, FSessionSettings> MemberSettings : Session->SessionSettings.MemberSettings)
+	SetLobbyMemberAttributes(LobbyModificationHandle, EOSSubsystem->UserManager->GetUniquePlayerId(EOSSubsystem->UserManager->GetDefaultLocalUser()).ToSharedRef(), *Session);
+}
+
+void FOnlineSessionEOS::SetLobbyMemberAttributes(EOS_HLobbyModification LobbyModificationHandle, FUniqueNetIdRef LobbyMemberId, FNamedOnlineSession& Session)
+{
+	if (FSessionSettings* MemberSettings = Session.SessionSettings.MemberSettings.Find(LobbyMemberId))
 	{
-		// We'll only copy our local player's attributes
-		if (*EOSSubsystem->UserManager->GetUniquePlayerId(EOSSubsystem->UserManager->GetDefaultLocalUser()) == *MemberSettings.Key)
+		for (FSessionSettings::TConstIterator It(*MemberSettings); It; ++It)
 		{
-			for (FSessionSettings::TConstIterator It(MemberSettings.Value); It; ++It)
-			{
-				const FName KeyName = It.Key();
-				const FOnlineSessionSetting& Setting = It.Value();
+			const FName KeyName = It.Key();
+			const FOnlineSessionSetting& Setting = It.Value();
 
-				// Skip unsupported types or non session advertised settings
-				if (Setting.AdvertisementType < EOnlineDataAdvertisementType::ViaOnlineService || !IsSessionSettingTypeSupported(Setting.Data.GetType()))
-				{
-					continue;
-				}
-
-				const FLobbyAttributeOptions Attribute(TCHAR_TO_UTF8(*KeyName.ToString()), Setting.Data);
-				AddLobbyMemberAttribute(LobbyModificationHandle, &Attribute);
+			// Skip unsupported types or non session advertised settings
+			if (Setting.AdvertisementType < EOnlineDataAdvertisementType::ViaOnlineService || !IsSessionSettingTypeSupported(Setting.Data.GetType()))
+			{
+				continue;
 			}
+
+			const FLobbyAttributeOptions Attribute(TCHAR_TO_UTF8(*KeyName.ToString()), Setting.Data);
+			AddLobbyMemberAttribute(LobbyModificationHandle, &Attribute);
 		}
 	}
+	else
+	{
+		UE_LOG_ONLINE_SESSION(Warning, TEXT("[FOnlineSessionEOS::SetLobbyMemberAttributes] Lobby Member with UniqueNetId [%s] not found in Lobby Session with SessionId [%s]"), *LobbyMemberId->ToDebugString(), *Session.GetSessionIdStr());
+	}
 }
 
 uint32 FOnlineSessionEOS::UpdateLobbySession(FNamedOnlineSession* Session, const FOnUpdateSessionCompleteDelegate& CompletionDelegate)
@@ -4060,9 +4064,17 @@ uint32 FOnlineSessionEOS::UpdateLobbySession(FNamedOnlineSession* Session, const
 		EOS_EResult LobbyModificationResult = EOS_Lobby_UpdateLobbyModification(LobbyHandle, &UpdateLobbyModificationOptions, &LobbyModificationHandle);
 		if (LobbyModificationResult == EOS_EResult::EOS_Success)
 		{
-			SetLobbyPermissionLevel(LobbyModificationHandle, Session);
-			SetLobbyMaxMembers(LobbyModificationHandle, Session);
-			SetLobbyAttributes(LobbyModificationHandle, Session);
+			// If the user initiating the update is the owner, we will update both lobby settings and member settings
+			if (FUniqueNetIdEOS::Cast(*Session->OwningUserId).GetProductUserId() == UpdateLobbyModificationOptions.LocalUserId)
+			{
+				SetLobbyPermissionLevel(LobbyModificationHandle, Session);
+				SetLobbyMaxMembers(LobbyModificationHandle, Session);
+				SetLobbyAttributes(LobbyModificationHandle, Session);
+			}
+			else // In any other case, only member settings will be updated, as per API restrictions
+			{
+				SetLobbyMemberAttributes(LobbyModificationHandle, EOSSubsystem->UserManager->GetUniquePlayerId(EOSSubsystem->UserManager->GetDefaultLocalUser()).ToSharedRef(), *Session);
+			}
 
 			EOS_Lobby_UpdateLobbyOptions UpdateLobbyOptions = { 0 };
 			UpdateLobbyOptions.ApiVersion = 1;
diff --git a/Engine/Plugins/Online/OnlineSubsystemEOS/Source/OnlineSubsystemEOS/Private/OnlineSessionEOS.h b/Engine/Plugins/Online/OnlineSubsystemEOS/Source/OnlineSubsystemEOS/Private/OnlineSessionEOS.h
--- a/Engine/Plugins/Online/OnlineSubsystemEOS/Source/OnlineSubsystemEOS/Private/OnlineSessionEOS.h
+++ b/Engine/Plugins/Online/OnlineSubsystemEOS/Source/OnlineSubsystemEOS/Private/OnlineSessionEOS.h
@@ -253,8 +253,8 @@ private:
 	void SetLobbyMaxMembers(EOS_HLobbyModification LobbyModificationHandle, FNamedOnlineSession* Session);
 	void SetLobbyAttributes(EOS_HLobbyModification LobbyModificationHandle, FNamedOnlineSession* Session);
 	void AddLobbyAttribute(EOS_HLobbyModification LobbyModificationHandle, const EOS_Lobby_AttributeData* Attribute);
+	void SetLobbyMemberAttributes(EOS_HLobbyModification LobbyModificationHandle, FUniqueNetIdRef LobbyMemberId, FNamedOnlineSession& Session);
 	void AddLobbyMemberAttribute(EOS_HLobbyModification LobbyModificationHandle, const EOS_Lobby_AttributeData* Attribute);
-	void AddLobbyMember(const FUniqueNetIdStringRef LobbyNetId, const EOS_ProductUserId& TargetUserId);
 
 	// Methods to update an OSS Lobby from an API Lobby
 	typedef TFunction<void(bool bWasSuccessful)> FOnCopyLobbyDataCompleteCallback;
@@ -269,12 +269,10 @@ private:
 	// Helper methods
 	typedef TFunction<void(const EOS_ProductUserId& ProductUserId, EOS_EpicAccountId& EpicAccountId)> GetEpicAccountIdAsyncCallback;
 
-	void GetEpicAccountIdAsync(const EOS_ProductUserId& ProductUserId, const GetEpicAccountIdAsyncCallback& Callback);
 	void RegisterLobbyNotifications();
 	FNamedOnlineSession* GetNamedSessionFromLobbyId(const FUniqueNetIdEOSLobby& LobbyId);
 	FOnlineSessionSearchResult* GetSearchResultFromLobbyId(const FUniqueNetIdEOSLobby& LobbyId);
 	FOnlineSession* GetOnlineSessionFromLobbyId(const FUniqueNetIdEOSLobby& LobbyId);
-	bool GetEpicAccountIdFromProductUserId(const EOS_ProductUserId& ProductUserId, EOS_EpicAccountId& EpicAccountId);
 	EOS_ELobbyPermissionLevel GetLobbyPermissionLevelFromSessionSettings(const FOnlineSessionSettings& SessionSettings);
 	uint32_t GetLobbyMaxMembersFromSessionSettings(const FOnlineSessionSettings& SessionSettings);
 
