Fix crash on editor startup.

diff --git a/Engine/Plugins/Online/OnlineSubsystemEOS/Source/OnlineSubsystemEOS/Private/OnlineSubsystemEOS.cpp b/Engine/Plugins/Online/OnlineSubsystemEOS/Source/OnlineSubsystemEOS/Private/OnlineSubsystemEOS.cpp
--- a/Engine/Plugins/Online/OnlineSubsystemEOS/Source/OnlineSubsystemEOS/Private/OnlineSubsystemEOS.cpp
+++ b/Engine/Plugins/Online/OnlineSubsystemEOS/Source/OnlineSubsystemEOS/Private/OnlineSubsystemEOS.cpp
@@ -381,7 +381,7 @@ bool FOnlineSubsystemEOS::Init()
 	}
 	else if (ArtifactSettings.EncryptionKey.IsEmpty())
 	{
-		UE_LOG_ONLINE(Verbose, TEXT("FOnlineSubsystemEOS: ArtifactName=[%s] EncryptionKey unset, TitleFile interface unavailable."));
+		UE_LOG_ONLINE(Verbose, TEXT("FOnlineSubsystemEOS: ArtifactName=[%s] EncryptionKey unset, TitleFile interface unavailable."), *ArtifactName);
 	}
 	else
 	{
@@ -396,7 +396,7 @@ bool FOnlineSubsystemEOS::Init()
 	}
 	else if (ArtifactSettings.EncryptionKey.IsEmpty())
 	{
-		UE_LOG_ONLINE(Verbose, TEXT("FOnlineSubsystemEOS: ArtifactName=[%s] EncryptionKey unset, UserCloud interface unavailable."));
+		UE_LOG_ONLINE(Verbose, TEXT("FOnlineSubsystemEOS: ArtifactName=[%s] EncryptionKey unset, UserCloud interface unavailable."), *ArtifactName);
 	}
 	else
 	{
