Refactor of OSS EOS logic surrounding setting the BucketId for Sessions and Lobbies on creation and searches

diff --git a/Engine/Plugins/Online/OnlineSubsystemEOS/Source/OnlineSubsystemEOS/Private/OnlineSessionEOS.cpp b/Engine/Plugins/Online/OnlineSubsystemEOS/Source/OnlineSubsystemEOS/Private/OnlineSessionEOS.cpp
--- a/Engine/Plugins/Online/OnlineSubsystemEOS/Source/OnlineSubsystemEOS/Private/OnlineSessionEOS.cpp
+++ b/Engine/Plugins/Online/OnlineSubsystemEOS/Source/OnlineSubsystemEOS/Private/OnlineSessionEOS.cpp
@@ -25,9 +25,6 @@
 #include "eos_metrics.h"
 #include "eos_lobby.h"
 
-/** This is the game name plus version in ansi done once for optimization */
-char BucketIdAnsi[EOS_OSS_BUCKET_ID_STRING_LENGTH + 1]; // X characters plus null terminator
-
 FString MakeStringFromAttributeValue(const EOS_Sessions_AttributeData* Attribute)
 {
 	switch (Attribute->ValueType)
@@ -458,10 +455,8 @@ FOnlineSessionEOS::~FOnlineSessionEOS()
 	delete JoinLobbyAcceptedCallback;
 }
 
-void FOnlineSessionEOS::Init(const FString& InBucketId)
+void FOnlineSessionEOS::Init()
 {
-	FCStringAnsi::Strncpy(BucketIdAnsi, TCHAR_TO_UTF8(*InBucketId), sizeof(BucketIdAnsi));
-
 	// Register for session invite notifications
 	FSessionInviteAcceptedCallback* SessionInviteAcceptedCallbackObj = new FSessionInviteAcceptedCallback(FOnlineSessionEOSWeakPtr(AsShared()));
 	SessionInviteAcceptedCallback = SessionInviteAcceptedCallbackObj;
@@ -1283,7 +1278,6 @@ struct FSessionCreateOptions :
 	{
 		ApiVersion = 4;
 		UE_EOS_CHECK_API_MISMATCH(EOS_SESSIONS_CREATESESSIONMODIFICATION_API_LATEST, 4);
-		BucketId = BucketIdAnsi;
 	}
 };
 
@@ -1300,6 +1294,8 @@ uint32 FOnlineSessionEOS::CreateEOSSession(int32 HostingPlayerNum, FNamedOnlineS
 		Session->SessionSettings.bAllowJoinViaPresence ||
 		Session->SessionSettings.bAllowJoinViaPresenceFriendsOnly ||
 		Session->SessionSettings.bAllowInvites) ? EOS_TRUE : EOS_FALSE;
+	const auto BucketIdUtf8 = StringCast<UTF8CHAR>(*GetBucketId(Session->SessionSettings));
+	Options.BucketId = (const char*)BucketIdUtf8.Get();
 
 	EOS_EResult ResultCode = EOS_Sessions_CreateSessionModification(EOSSubsystem->SessionsHandle, &Options, &SessionModHandle);
 	if (ResultCode != EOS_EResult::EOS_Success)
@@ -2213,7 +2209,8 @@ uint32 FOnlineSessionEOS::FindEOSSession(int32 SearchingPlayerNum, const TShared
 	FAttributeOptions Opt1("NumPublicConnections", 1);
 	AddSearchAttribute(SearchHandle, &Opt1, EOS_EOnlineComparisonOp::EOS_OCO_GREATERTHANOREQUAL);
 
-	FAttributeOptions Opt2(EOS_SESSIONS_SEARCH_BUCKET_ID, BucketIdAnsi);
+	const auto BucketIdUtf8 = StringCast<UTF8CHAR>(*GetBucketId(*SearchSettings));
+	FAttributeOptions Opt2(EOS_SESSIONS_SEARCH_BUCKET_ID, (const char*)BucketIdUtf8.Get());
 	AddSearchAttribute(SearchHandle, &Opt2, EOS_EOnlineComparisonOp::EOS_OCO_EQUAL);
 
 	// Add the search settings
@@ -3634,6 +3631,48 @@ uint32_t FOnlineSessionEOS::GetLobbyMaxMembersFromSessionSettings(const FOnlineS
 	return SessionSettings.NumPrivateConnections + SessionSettings.NumPublicConnections;
 }
 
+FString FOnlineSessionEOS::GetBucketId(const FOnlineSessionSettings& SessionSettings)
+{
+	// Check if the Bucket Id custom setting is set, set default otherwise. EOS Sessions and Lobbies can not be created without it
+	FString BucketIdStr;
+
+	if (const FOnlineSessionSetting* BucketIdSetting = SessionSettings.Settings.Find(OSSEOS_BUCKET_ID_ATTRIBUTE_KEY))
+	{
+		BucketIdSetting->Data.GetValue(BucketIdStr);
+	}
+	else
+	{
+		const int32 BuildUniqueId = GetBuildUniqueId();
+
+		UE_LOG_ONLINE_SESSION(Verbose, TEXT("[FOnlineSessionEOS::GetBucketId] 'OSSEOS_BUCKET_ID_ATTRIBUTE_KEY' (FString) Custom Setting needed to create EOS sessions not found. Setting \"%d\" as default."), BuildUniqueId);
+
+		BucketIdStr = FString::FromInt(BuildUniqueId);
+	}
+
+	return BucketIdStr;
+}
+
+FString FOnlineSessionEOS::GetBucketId(const FOnlineSessionSearch& SessionSearch)
+{
+	// Check if the Bucket Id filter is set, set default otherwise.
+	FString BucketIdStr;
+
+	if (const FOnlineSessionSearchParam* BucketIdSetting = SessionSearch.QuerySettings.SearchParams.Find(OSSEOS_BUCKET_ID_ATTRIBUTE_KEY))
+	{
+		BucketIdSetting->Data.GetValue(BucketIdStr);
+	}
+	else
+	{
+		const int32 BuildUniqueId = GetBuildUniqueId();
+
+		UE_LOG_ONLINE_SESSION(Verbose, TEXT("[FOnlineSessionEOS::GetBucketId] 'OSSEOS_BUCKET_ID_ATTRIBUTE_KEY' (FString) Custom Setting used to find EOS sessions not found. Setting \"%d\" as default."), BuildUniqueId);
+
+		BucketIdStr = FString::FromInt(BuildUniqueId);
+	}
+
+	return BucketIdStr;
+}
+
 uint32 FOnlineSessionEOS::CreateLobbySession(int32 HostingPlayerNum, FNamedOnlineSession* Session)
 {
 	check(Session != nullptr);
@@ -3654,7 +3693,8 @@ uint32 FOnlineSessionEOS::CreateLobbySession(int32 HostingPlayerNum, FNamedOnlin
 	CreateLobbyOptions.PermissionLevel = GetLobbyPermissionLevelFromSessionSettings(Session->SessionSettings);
 	CreateLobbyOptions.bPresenceEnabled = Session->SessionSettings.bUsesPresence;
 	CreateLobbyOptions.bAllowInvites = Session->SessionSettings.bAllowInvites;
-	CreateLobbyOptions.BucketId = BucketIdAnsi;
+	const auto BucketIdUtf8 = StringCast<UTF8CHAR>(*GetBucketId(Session->SessionSettings));
+	CreateLobbyOptions.BucketId = (const char*)BucketIdUtf8.Get();
 	CreateLobbyOptions.bDisableHostMigration = !bUseHostMigration;
 #if WITH_EOS_RTC
 	CreateLobbyOptions.bEnableRTCRoom = Session->SessionSettings.bUseLobbiesVoiceChatIfAvailable;
diff --git a/Engine/Plugins/Online/OnlineSubsystemEOS/Source/OnlineSubsystemEOS/Private/OnlineSessionEOS.h b/Engine/Plugins/Online/OnlineSubsystemEOS/Source/OnlineSubsystemEOS/Private/OnlineSessionEOS.h
--- a/Engine/Plugins/Online/OnlineSubsystemEOS/Source/OnlineSubsystemEOS/Private/OnlineSessionEOS.h
+++ b/Engine/Plugins/Online/OnlineSubsystemEOS/Source/OnlineSubsystemEOS/Private/OnlineSessionEOS.h
@@ -117,7 +117,10 @@ public:
 	}
 
 // IOnlineSession Interface
+
+	/** The BucketId for the session can be specified via adding a custom session setting with the key OSSEOS_BUCKET_ID_ATTRIBUTE_KEY. A default value will be used otherwise */
 	virtual bool CreateSession(int32 HostingPlayerNum, FName SessionName, const FOnlineSessionSettings& NewSessionSettings) override;
+	/** The BucketId for the session can be specified via adding a custom session setting with the key OSSEOS_BUCKET_ID_ATTRIBUTE_KEY. A default value will be used otherwise */
 	virtual bool CreateSession(const FUniqueNetId& HostingPlayerId, FName SessionName, const FOnlineSessionSettings& NewSessionSettings) override;
 	virtual bool StartSession(FName SessionName) override;
 	virtual bool UpdateSession(FName SessionName, FOnlineSessionSettings& UpdatedSessionSettings, bool bShouldRefreshOnlineData = true) override;
@@ -127,7 +130,9 @@ public:
 	virtual bool StartMatchmaking(const TArray< FUniqueNetIdRef >& LocalPlayers, FName SessionName, const FOnlineSessionSettings& NewSessionSettings, TSharedRef<FOnlineSessionSearch>& SearchSettings) override;
 	virtual bool CancelMatchmaking(int32 SearchingPlayerNum, FName SessionName) override;
 	virtual bool CancelMatchmaking(const FUniqueNetId& SearchingPlayerId, FName SessionName) override;
+	/** The BucketId to be used in the search can be specified via adding a custom search filter with the key OSSEOS_BUCKET_ID_ATTRIBUTE_KEY. A default value will be used otherwise */
 	virtual bool FindSessions(int32 SearchingPlayerNum, const TSharedRef<FOnlineSessionSearch>& SearchSettings) override;
+	/** The BucketId to be used in the search can be specified via adding a custom search filter with the key OSSEOS_BUCKET_ID_ATTRIBUTE_KEY. A default value will be used otherwise */
 	virtual bool FindSessions(const FUniqueNetId& SearchingPlayerId, const TSharedRef<FOnlineSessionSearch>& SearchSettings) override;
 	virtual bool FindSessionById(const FUniqueNetId& SearchingUserId, const FUniqueNetId& SessionId, const FUniqueNetId& FriendId, const FOnSingleSessionResultCompleteDelegate& CompletionDelegate) override;
 	virtual bool CancelFindSessions() override;
@@ -201,7 +206,7 @@ public:
 
 	void RegisterLocalPlayers(class FNamedOnlineSession* Session);
 
-	void Init(const FString& InBucketId);
+	void Init();
 
 private:
 	// EOS Lobbies
@@ -275,6 +280,8 @@ private:
 	FOnlineSession* GetOnlineSessionFromLobbyId(const FUniqueNetIdEOSLobby& LobbyId);
 	EOS_ELobbyPermissionLevel GetLobbyPermissionLevelFromSessionSettings(const FOnlineSessionSettings& SessionSettings);
 	uint32_t GetLobbyMaxMembersFromSessionSettings(const FOnlineSessionSettings& SessionSettings);
+	static FString GetBucketId(const FOnlineSessionSettings& SessionSettings);
+	static FString GetBucketId(const FOnlineSessionSearch& SessionSettings);
 
 	// EOS Sessions
 	uint32 CreateEOSSession(int32 HostingPlayerNum, FNamedOnlineSession* Session);
diff --git a/Engine/Plugins/Online/OnlineSubsystemEOS/Source/OnlineSubsystemEOS/Private/OnlineSubsystemEOS.cpp b/Engine/Plugins/Online/OnlineSubsystemEOS/Source/OnlineSubsystemEOS/Private/OnlineSubsystemEOS.cpp
--- a/Engine/Plugins/Online/OnlineSubsystemEOS/Source/OnlineSubsystemEOS/Private/OnlineSubsystemEOS.cpp
+++ b/Engine/Plugins/Online/OnlineSubsystemEOS/Source/OnlineSubsystemEOS/Private/OnlineSubsystemEOS.cpp
@@ -415,17 +415,7 @@ bool FOnlineSubsystemEOS::Init()
 	UserManager = MakeShareable(new FUserManagerEOS(this));
 	UserManager->Init();
 	SessionInterfacePtr = MakeShareable(new FOnlineSessionEOS(this));
-	// Set the bucket id to use for all sessions based upon the Build Id to avoid upgrade issues
-	FString BucketIdStr = FString::FromInt(GetBuildUniqueId());
-	if (BucketIdStr.Len() > EOS_OSS_BUCKET_ID_STRING_LENGTH)
-	{
-		FString NewBucketIdStr = BucketIdStr.Left(EOS_OSS_BUCKET_ID_STRING_LENGTH);
-
-		UE_LOG_ONLINE(Warning, TEXT("[FOnlineSubsystemEOS::Init] Default BucketId [%s] is too long (%d characters) and will be shortened to fit the limit (%d characters). New BucketId is [%s]"), *BucketIdStr, BucketIdStr.Len(), EOS_OSS_BUCKET_ID_STRING_LENGTH, *NewBucketIdStr);
-
-		BucketIdStr = MoveTemp(NewBucketIdStr);		
-	}
-	SessionInterfacePtr->Init(BucketIdStr);
+	SessionInterfacePtr->Init();
 	StatsInterfacePtr = MakeShareable(new FOnlineStatsEOS(this));
 	LeaderboardsInterfacePtr = MakeShareable(new FOnlineLeaderboardsEOS(this));
 	AchievementsInterfacePtr = MakeShareable(new FOnlineAchievementsEOS(this));
diff --git a/Engine/Plugins/Online/OnlineSubsystemEOS/Source/OnlineSubsystemEOS/Public/OnlineSubsystemEOSTypesPublic.h b/Engine/Plugins/Online/OnlineSubsystemEOS/Source/OnlineSubsystemEOS/Public/OnlineSubsystemEOSTypesPublic.h
--- a/Engine/Plugins/Online/OnlineSubsystemEOS/Source/OnlineSubsystemEOS/Public/OnlineSubsystemEOSTypesPublic.h
+++ b/Engine/Plugins/Online/OnlineSubsystemEOS/Source/OnlineSubsystemEOS/Public/OnlineSubsystemEOSTypesPublic.h
@@ -28,4 +28,6 @@ public:
 	virtual const EOS_ProductUserId GetProductUserId() const = 0;
 };
 
+#define OSSEOS_BUCKET_ID_ATTRIBUTE_KEY FName("OSSEOS_BUCKET_ID_ATTRIBUTE_KEY")
+
 #endif // WITH_EOS_SDK
\ No newline at end of file
