- Add Linux to whitelisted OSSEOS platforms

diff --git a/Engine/Plugins/Online/OnlineSubsystemEOS/OnlineSubsystemEOS.uplugin b/Engine/Plugins/Online/OnlineSubsystemEOS/OnlineSubsystemEOS.uplugin
--- a/Engine/Plugins/Online/OnlineSubsystemEOS/OnlineSubsystemEOS.uplugin
+++ b/Engine/Plugins/Online/OnlineSubsystemEOS/OnlineSubsystemEOS.uplugin
@@ -19,6 +19,7 @@
 			"PlatformAllowList": [
 				"Win64",
 				"Mac",
+				"Linux",
 				"Android"
 			]
 		},
@@ -29,6 +30,7 @@
 			"PlatformAllowList": [
 				"Win64",
 				"Mac",
+				"Linux",
 				"Android"
 			]
 		}
diff --git a/Engine/Plugins/Online/OnlineSubsystemEOS/Source/OnlineSubsystemEOS/Private/Linux/LinuxEOSHelpers.h b/Engine/Plugins/Online/OnlineSubsystemEOS/Source/OnlineSubsystemEOS/Private/Linux/LinuxEOSHelpers.h
new file mode 100644
--- /dev/null
+++ b/Engine/Plugins/Online/OnlineSubsystemEOS/Source/OnlineSubsystemEOS/Private/Linux/LinuxEOSHelpers.h
@@ -0,0 +1,12 @@
+// Copyright Epic Games, Inc. All Rights Reserved.
+
+#pragma once
+
+#if WITH_EOS_SDK
+
+#include "EOSHelpers.h"
+
+using FPlatformEOSHelpers = FEOSHelpers;
+
+#endif
+
