Removed early exit in FriendStatusChanged method in order to handle accepted friend invites via the overlay.

diff --git a/Engine/Plugins/Online/OnlineSubsystemEOS/Source/OnlineSubsystemEOS/Private/UserManagerEOS.cpp b/Engine/Plugins/Online/OnlineSubsystemEOS/Source/OnlineSubsystemEOS/Private/UserManagerEOS.cpp
--- a/Engine/Plugins/Online/OnlineSubsystemEOS/Source/OnlineSubsystemEOS/Private/UserManagerEOS.cpp
+++ b/Engine/Plugins/Online/OnlineSubsystemEOS/Source/OnlineSubsystemEOS/Private/UserManagerEOS.cpp
@@ -1768,27 +1768,21 @@ typedef TEOSCallback<EOS_Friends_OnQueryFriendsCallback, EOS_Friends_QueryFriend
 
 void FUserManagerEOS::FriendStatusChanged(const EOS_Friends_OnFriendsUpdateInfo* Data)
 {
-	// This seems to happen due to the SDK's local cache going from empty to filled, so ignore it
-	// It's not really a valid transition since there should have been a pending invite inbetween
-	if (Data->PreviousStatus == EOS_EFriendsStatus::EOS_FS_NotFriends && Data->CurrentStatus == EOS_EFriendsStatus::EOS_FS_Friends)
-	{
-		return;
-	}
-
 	// Get the local user information
 	if (AccountIdToUserNumMap.Contains(Data->LocalUserId))
 	{
 		int32 LocalUserNum = AccountIdToUserNumMap[Data->LocalUserId];
 		FUniqueNetIdEOSPtr LocalEOSID = UserNumToNetIdMap[LocalUserNum];
-		// If we don't know them yet, then add them to kick off the reads
-		if (!AccountIdToStringMap.Contains(Data->TargetUserId))
+		// If we don't know them yet, or if they just became a friend, then add them to kick off the reads
+		if (!AccountIdToStringMap.Contains(Data->TargetUserId) || Data->CurrentStatus == EOS_EFriendsStatus::EOS_FS_Friends)
 		{
 			AddFriend(LocalUserNum, Data->TargetUserId);
 		}
 		// They are in our list now
 		FOnlineUserPtr OnlineUser = EpicAccountIdToOnlineUserMap[Data->TargetUserId];
 		FOnlineFriendEOSPtr Friend = LocalUserNumToFriendsListMap[LocalUserNum]->GetByNetIdString(AccountIdToStringMap[Data->TargetUserId]);
-		// Figure out which notification to fire
+
+		// Figure out which notification to fire. Invite related notifications will only fire in projects with friends management enabled
 		if (Data->CurrentStatus == EOS_EFriendsStatus::EOS_FS_Friends)
 		{
 			Friend->SetInviteStatus(EInviteStatus::Accepted);
