Delayed callback execution to match alternate flow paths in EOS Session search

diff --git a/Engine/Plugins/Online/OnlineSubsystemEOS/Source/OnlineSubsystemEOS/Private/OnlineSessionEOS.cpp b/Engine/Plugins/Online/OnlineSubsystemEOS/Source/OnlineSubsystemEOS/Private/OnlineSessionEOS.cpp
--- a/Engine/Plugins/Online/OnlineSubsystemEOS/Source/OnlineSubsystemEOS/Private/OnlineSessionEOS.cpp
+++ b/Engine/Plugins/Online/OnlineSubsystemEOS/Source/OnlineSubsystemEOS/Private/OnlineSessionEOS.cpp
@@ -4426,9 +4426,13 @@ void FOnlineSessionEOS::AddLobbySearchResult(const TSharedRef<FLobbyDetailsEOS>&
 	}
 	else
 	{
-		UE_LOG_ONLINE_SESSION(Warning, TEXT("[FOnlineSessionEOS::AddLobbySearchResult] LobbyDetails_CopyInfo not successful. Finished with EOS_EResult %s"), ANSI_TO_TCHAR(EOS_EResult_ToString(CopyResult)));
+		// CopyLobbyData may launch an asynchronous operation, so we'll delay the execution of this callback to match the flow
+		EOSSubsystem->ExecuteNextTick([CopyResult, Callback]()
+		{
+			UE_LOG_ONLINE_SESSION(Warning, TEXT("[FOnlineSessionEOS::AddLobbySearchResult] LobbyDetails_CopyInfo not successful. Finished with EOS_EResult %s"), *LexToString(CopyResult));
 
-		Callback(false);
+			Callback(false);
+		});
 	}
 }
 
@@ -4522,7 +4526,11 @@ void FOnlineSessionEOS::CopyLobbyData(const TSharedRef<FLobbyDetailsEOS>& LobbyD
 	}
 	else
 	{
-		Callback(true);
+		// ResolveUniqueNetIds is an asynchronous operation, so we'll delay the execution of this callback to match the flow
+		EOSSubsystem->ExecuteNextTick([Callback]()
+		{
+			Callback(true);
+		});
 	}
 }
 
