Added bespoke assignments of nullptr to unused parameters in EOS login calls

diff --git a/Engine/Plugins/Online/OnlineSubsystemEOS/Source/OnlineSubsystemEOS/Private/UserManagerEOS.cpp b/Engine/Plugins/Online/OnlineSubsystemEOS/Source/OnlineSubsystemEOS/Private/UserManagerEOS.cpp
--- a/Engine/Plugins/Online/OnlineSubsystemEOS/Source/OnlineSubsystemEOS/Private/UserManagerEOS.cpp
+++ b/Engine/Plugins/Online/OnlineSubsystemEOS/Source/OnlineSubsystemEOS/Private/UserManagerEOS.cpp
@@ -402,7 +402,7 @@ struct FAuthCredentials :
 		UE_EOS_CHECK_API_MISMATCH(EOS_AUTH_CREDENTIALS_API_LATEST, 3);
 		Type = EOS_ELoginCredentialType::EOS_LCT_ExternalAuth;
 		ExternalType = InExternalType;
-		Id = IdAnsi;
+		Id = nullptr;
 		Token = TokenAnsi;
 
 		FCStringAnsi::Strncpy(TokenAnsi, TCHAR_TO_UTF8(*InTokenString), InTokenString.Len()+1);
@@ -414,7 +414,7 @@ struct FAuthCredentials :
 		UE_EOS_CHECK_API_MISMATCH(EOS_AUTH_CREDENTIALS_API_LATEST, 3);
 		Type = EOS_ELoginCredentialType::EOS_LCT_ExternalAuth;
 		ExternalType = InExternalType;
-		Id = IdAnsi;
+		Id = nullptr;
 		Token = TokenAnsi;
 
 		uint32_t InOutBufferLength = EOS_MAX_TOKEN_SIZE;
@@ -478,8 +478,9 @@ bool FUserManagerEOS::Login(int32 LocalUserNum, const FOnlineAccountCredentials&
 	if (AccountCredentials.Type == TEXT("exchangecode"))
 	{
 		// This is how the Epic launcher will pass credentials to you
-		FCStringAnsi::Strncpy(Credentials.TokenAnsi, TCHAR_TO_UTF8(*AccountCredentials.Token), EOS_MAX_TOKEN_SIZE);
 		Credentials.Type = EOS_ELoginCredentialType::EOS_LCT_ExchangeCode;
+		Credentials.Id = nullptr;
+		FCStringAnsi::Strncpy(Credentials.TokenAnsi, TCHAR_TO_UTF8(*AccountCredentials.Token), EOS_MAX_TOKEN_SIZE);
 	}
 	else if (AccountCredentials.Type == TEXT("developer"))
 	{
@@ -499,13 +500,13 @@ bool FUserManagerEOS::Login(int32 LocalUserNum, const FOnlineAccountCredentials&
 	{
 		// This is auth via the EOS Account Portal
 		Credentials.Type = EOS_ELoginCredentialType::EOS_LCT_AccountPortal;
+		Credentials.Id = nullptr;
+		Credentials.Token = nullptr;
 	}
 	else if (AccountCredentials.Type == TEXT("persistentauth"))
 	{
 		// Use locally stored token managed by EOSSDK keyring to attempt login.
 		Credentials.Type = EOS_ELoginCredentialType::EOS_LCT_PersistentAuth;
-
-		// Id and Token must be null when using EOS_ELoginCredentialType::EOS_LCT_PersistentAuth
 		Credentials.Id = nullptr;
 		Credentials.Token = nullptr;
 
