Added logic to skip user id retrieval in UpdateStats when running on a Dedicated Server

diff --git a/Engine/Plugins/Online/OnlineSubsystemEOS/Source/OnlineSubsystemEOS/Private/OnlineStatsEOS.cpp b/Engine/Plugins/Online/OnlineSubsystemEOS/Source/OnlineSubsystemEOS/Private/OnlineStatsEOS.cpp
--- a/Engine/Plugins/Online/OnlineSubsystemEOS/Source/OnlineSubsystemEOS/Private/OnlineStatsEOS.cpp
+++ b/Engine/Plugins/Online/OnlineSubsystemEOS/Source/OnlineSubsystemEOS/Private/OnlineStatsEOS.cpp
@@ -324,14 +324,19 @@ void FOnlineStatsEOS::WriteStats(EOS_ProductUserId LocalUserId, EOS_ProductUserI
 
 void FOnlineStatsEOS::UpdateStats(const FUniqueNetIdRef LocalUserId, const TArray<FOnlineStatsUserUpdatedStats>& UpdatedUserStats, const FOnlineStatsUpdateStatsComplete& Delegate)
 {
-	const FUniqueNetIdEOS& EOSId = FUniqueNetIdEOS::Cast(*LocalUserId);
-	const EOS_ProductUserId UserId = EOSId.GetProductUserId();
-	if (UserId == nullptr)
+	EOS_ProductUserId UserId = nullptr;
+	if (!IsRunningDedicatedServer()) // We'll keep the calling user as null for dedicated servers
 	{
-		UE_LOG_ONLINE_STATS(Error, TEXT("UpdateStats() failed for unknown player (%s)"), *EOSId.ToDebugString());
-		Delegate.ExecuteIfBound(FOnlineError(EOnlineErrorResult::InvalidCreds));
-		return;
+		const FUniqueNetIdEOS& EOSId = FUniqueNetIdEOS::Cast(*LocalUserId);
+		UserId = EOSId.GetProductUserId();
+		if (UserId == nullptr)
+		{
+			UE_LOG_ONLINE_STATS(Error, TEXT("UpdateStats() failed for unknown player (%s)"), *EOSId.ToDebugString());
+			Delegate.ExecuteIfBound(FOnlineError(EOnlineErrorResult::InvalidCreds));
+			return;
+		}
 	}
+
 	if (UpdatedUserStats.Num() == 0 )
 	{
 		UE_LOG_ONLINE_STATS(Error, TEXT("UpdateStats() failed for player due to no user stats being specified"));
