- Move warning to WriteAchievements

diff --git a/Engine/Plugins/Online/OnlineSubsystemEOS/Source/OnlineSubsystemEOS/Private/OnlineAchievementsEOS.cpp b/Engine/Plugins/Online/OnlineSubsystemEOS/Source/OnlineSubsystemEOS/Private/OnlineAchievementsEOS.cpp
--- a/Engine/Plugins/Online/OnlineSubsystemEOS/Source/OnlineSubsystemEOS/Private/OnlineAchievementsEOS.cpp
+++ b/Engine/Plugins/Online/OnlineSubsystemEOS/Source/OnlineSubsystemEOS/Private/OnlineAchievementsEOS.cpp
@@ -21,11 +21,6 @@ FOnlineAchievementsEOS::FOnlineAchievementsEOS(FOnlineSubsystemEOS* InSubsystem)
 void FOnlineAchievementsEOS::Init()
 {
 	GConfig->GetBool(TEXT("OnlineSubsystemEOS"), TEXT("bUseUnlockAchievements"), bUseUnlockAchievements, GEngineIni);
-
-	if (!bUseUnlockAchievements)
-	{
-		UE_LOG_ONLINE_ACHIEVEMENTS(Warning, TEXT("Upgrade note: OSSEOS is updating its implementation of WriteAchievements to match other OSS implementations. For more information, search for bUseUnlockAchievements in the release notes."));
-	}
 }
 
 void FOnlineAchievementsEOS::WriteAchievements(const FUniqueNetId& PlayerId, FOnlineAchievementsWriteRef& WriteObject, const FOnAchievementsWrittenDelegate& Delegate)
@@ -36,6 +31,7 @@ void FOnlineAchievementsEOS::WriteAchievements(const FUniqueNetId& PlayerId, FOn
 	}
 	else // If bUseUnlockAchievements is not set to true, we'll unlock the achievements through stat changes
 	{
+		UE_LOG_ONLINE_ACHIEVEMENTS(Warning, TEXT("Upgrade note: OSSEOS is updating its implementation of WriteAchievements to match other OSS implementations. For more information, search for bUseUnlockAchievements in the release notes."));
 PRAGMA_DISABLE_DEPRECATION_WARNINGS
 		UnlockAchievementsThroughStats(PlayerId, WriteObject, Delegate);
 PRAGMA_ENABLE_DEPRECATION_WARNINGS
