- Merge UnlockAchievements support and deprecation for old flow to 5.2

diff --git a/Engine/Plugins/Online/OnlineSubsystemEOS/Source/OnlineSubsystemEOS/Private/OnlineAchievementsEOS.cpp b/Engine/Plugins/Online/OnlineSubsystemEOS/Source/OnlineSubsystemEOS/Private/OnlineAchievementsEOS.cpp
--- a/Engine/Plugins/Online/OnlineSubsystemEOS/Source/OnlineSubsystemEOS/Private/OnlineAchievementsEOS.cpp
+++ b/Engine/Plugins/Online/OnlineSubsystemEOS/Source/OnlineSubsystemEOS/Private/OnlineAchievementsEOS.cpp
@@ -83,7 +83,7 @@ void FOnlineAchievementsEOS::UnlockAchievements(const FUniqueNetId& PlayerId, FO
 				{
 					for (FOnlineAchievement& Achievement : Achievements->Get())
 					{
-						if (Achievement.Id == AchievementId)
+						if (Achievement.Id == AchievementId.ToString())
 						{
 							Achievement.Progress = 1.0f;
 							TriggerOnAchievementUnlockedDelegates(*LambdaPlayerId, AchievementId.ToString());
