- Stop FUniqueNetIdEOSPlus inheriting from FUniqueNetIdString
- Improve FUniqueNetIdEOSPlus::ToDebugString to redact each part separately.

diff --git a/Engine/Plugins/Online/OnlineSubsystemEOS/Source/OnlineSubsystemEOSPlus/Private/OnlineUserEOSPlus.cpp b/Engine/Plugins/Online/OnlineSubsystemEOS/Source/OnlineSubsystemEOSPlus/Private/OnlineUserEOSPlus.cpp
--- a/Engine/Plugins/Online/OnlineSubsystemEOS/Source/OnlineSubsystemEOSPlus/Private/OnlineUserEOSPlus.cpp
+++ b/Engine/Plugins/Online/OnlineSubsystemEOS/Source/OnlineSubsystemEOSPlus/Private/OnlineUserEOSPlus.cpp
@@ -58,17 +58,8 @@ static inline FName ToOSSName(EOSSValue OSSValue)
 	return NULL_SUBSYSTEM;
 }
 
-inline FString BuildEOSPlusStringId(FUniqueNetIdPtr InBaseUniqueNetId, FUniqueNetIdPtr InEOSUniqueNetId)
-{
-	FString StrId = InBaseUniqueNetId.IsValid() ? InBaseUniqueNetId->ToString() : TEXT("");
-	StrId += EOSPLUS_ID_SEPARATOR;
-	StrId += InEOSUniqueNetId.IsValid() ? InEOSUniqueNetId->ToString() : TEXT("");
-	return StrId;
-}
-
 FUniqueNetIdEOSPlus::FUniqueNetIdEOSPlus(FUniqueNetIdPtr InBaseUniqueNetId, FUniqueNetIdPtr InEOSUniqueNetId)
-	: FUniqueNetIdString(BuildEOSPlusStringId(InBaseUniqueNetId, InEOSUniqueNetId), FName("EOSPlus"))
-	, BaseUniqueNetId(InBaseUniqueNetId)
+	: BaseUniqueNetId(InBaseUniqueNetId)
 	, EOSUniqueNetId(InEOSUniqueNetId)
 {
 	int32 TotalBytes = GetSize();
@@ -97,6 +88,12 @@ FUniqueNetIdEOSPlus::FUniqueNetIdEOSPlus(FUniqueNetIdPtr InBaseUniqueNetId, FUni
 	}
 }
 
+FName FUniqueNetIdEOSPlus::GetType() const
+{
+	static FName Type(TEXT("EOSPlus"));
+	return Type;
+}
+
 const uint8* FUniqueNetIdEOSPlus::GetBytes() const
 {
 	return RawBytes.GetData();
@@ -118,6 +115,20 @@ bool FUniqueNetIdEOSPlus::IsValid() const
 	return BaseUniqueNetId.IsValid() || EOSUniqueNetId.IsValid();
 }
 
+FString FUniqueNetIdEOSPlus::ToString() const
+{
+	const FString BaseStr = BaseUniqueNetId.IsValid() ? BaseUniqueNetId->ToString() : FString();
+	const FString EosStr = EOSUniqueNetId.IsValid() ? EOSUniqueNetId->ToString() : FString();
+	return BaseStr + EOSPLUS_ID_SEPARATOR + EosStr;
+}
+
+FString FUniqueNetIdEOSPlus::ToDebugString() const
+{
+	const FString BaseStr = BaseUniqueNetId.IsValid() ? BaseUniqueNetId->ToDebugString() : TEXT("INVALID");
+	const FString EosStr = EOSUniqueNetId.IsValid() ? EOSUniqueNetId->ToDebugString() : TEXT("INVALID");
+	return BaseStr + EOSPLUS_ID_SEPARATOR + EosStr;
+}
+
 FOnlineUserEOSPlus::FOnlineUserEOSPlus(FOnlineSubsystemEOSPlus* InSubsystem)
 	: EOSPlus(InSubsystem)
 {
diff --git a/Engine/Plugins/Online/OnlineSubsystemEOS/Source/OnlineSubsystemEOSPlus/Private/OnlineUserEOSPlus.h b/Engine/Plugins/Online/OnlineSubsystemEOS/Source/OnlineSubsystemEOSPlus/Private/OnlineUserEOSPlus.h
--- a/Engine/Plugins/Online/OnlineSubsystemEOS/Source/OnlineSubsystemEOSPlus/Private/OnlineUserEOSPlus.h
+++ b/Engine/Plugins/Online/OnlineSubsystemEOS/Source/OnlineSubsystemEOSPlus/Private/OnlineUserEOSPlus.h
@@ -26,8 +26,7 @@ using FUniqueNetIdEOSPlusRef = TSharedRef<const class FUniqueNetIdEOSPlus>;
  * Unique net id wrapper for a EOS plus another account id. The underlying string is a combination
  * of both account ids concatenated
  */
-class FUniqueNetIdEOSPlus :
-	public FUniqueNetIdString
+class FUniqueNetIdEOSPlus : public FUniqueNetId
 {
 public:
 	template<typename... TArgs>
@@ -37,9 +36,12 @@ public:
 	}
 
 // FUniqueNetId interface
+	virtual FName GetType() const override;
 	virtual const uint8* GetBytes() const override;
 	virtual int32 GetSize() const override;
 	virtual bool IsValid() const override;
+	virtual FString ToString() const override;
+	virtual FString ToDebugString() const override;
 // ~FUniqueNetId interface
 
 	FUniqueNetIdPtr GetBaseNetId() const
