EOS Auth refactor to remove unnecessary types and avoid passing empty strings to API call

diff --git a/Engine/Plugins/Online/OnlineSubsystemEOS/Source/OnlineSubsystemEOS/Private/EOSHelpers.cpp b/Engine/Plugins/Online/OnlineSubsystemEOS/Source/OnlineSubsystemEOS/Private/EOSHelpers.cpp
--- a/Engine/Plugins/Online/OnlineSubsystemEOS/Source/OnlineSubsystemEOS/Private/EOSHelpers.cpp
+++ b/Engine/Plugins/Online/OnlineSubsystemEOS/Source/OnlineSubsystemEOS/Private/EOSHelpers.cpp
@@ -10,7 +10,7 @@
 #include "OnlineSubsystemEOSPrivate.h"
 #include "UserManagerEOS.h"
 
-void FEOSHelpers::PlatformAuthCredentials(EOS_Auth_Credentials &Credentials)
+void FEOSHelpers::GetSystemAuthCredentialsOptions(void*& OutCredentialsOptions)
 {
 }
 
diff --git a/Engine/Plugins/Online/OnlineSubsystemEOS/Source/OnlineSubsystemEOS/Private/EOSHelpers.h b/Engine/Plugins/Online/OnlineSubsystemEOS/Source/OnlineSubsystemEOS/Private/EOSHelpers.h
--- a/Engine/Plugins/Online/OnlineSubsystemEOS/Source/OnlineSubsystemEOS/Private/EOSHelpers.h
+++ b/Engine/Plugins/Online/OnlineSubsystemEOS/Source/OnlineSubsystemEOS/Private/EOSHelpers.h
@@ -20,7 +20,7 @@ class FEOSHelpers
 public:
 	virtual ~FEOSHelpers() = default;
 
-	virtual void PlatformAuthCredentials(EOS_Auth_Credentials &Credentials);
+	virtual void GetSystemAuthCredentialsOptions(void*& OutCredentialsOptions);
 	virtual void PlatformTriggerLoginUI(FOnlineSubsystemEOS* EOSSubsystem, const int ControllerIndex, bool bShowOnlineOnly, bool bShowSkipButton, const FOnLoginUIClosedDelegate& Delegate);
 	virtual IEOSPlatformHandlePtr CreatePlatform(EOS_Platform_Options& PlatformOptions);
 
diff --git a/Engine/Plugins/Online/OnlineSubsystemEOS/Source/OnlineSubsystemEOS/Private/IOS/IOSEOSHelpers.cpp b/Engine/Plugins/Online/OnlineSubsystemEOS/Source/OnlineSubsystemEOS/Private/IOS/IOSEOSHelpers.cpp
--- a/Engine/Plugins/Online/OnlineSubsystemEOS/Source/OnlineSubsystemEOS/Private/IOS/IOSEOSHelpers.cpp
+++ b/Engine/Plugins/Online/OnlineSubsystemEOS/Source/OnlineSubsystemEOS/Private/IOS/IOSEOSHelpers.cpp
@@ -53,7 +53,7 @@ FIOSEOSHelpers::~FIOSEOSHelpers()
 #endif
 }
 
-void FIOSEOSHelpers::PlatformAuthCredentials(EOS_Auth_Credentials &Credentials)
+void FIOSEOSHelpers::GetSystemAuthCredentialsOptions(void*& OutCredentialsOptions)
 {
 #if PRESENTATIONCONTEXT_REQUIRED
 	if (PresentationContextProvider == nil)
@@ -70,7 +70,7 @@ void FIOSEOSHelpers::PlatformAuthCredentials(EOS_Auth_Credentials &Credentials)
 	CredentialsOptions.PresentationContextProviding = nullptr;
 #endif
 
-	Credentials.SystemAuthCredentialsOptions = (void*)&CredentialsOptions;
+	OutCredentialsOptions = (void*)&CredentialsOptions;
 }
 
 void FIOSEOSHelpers::PlatformTriggerLoginUI(FOnlineSubsystemEOS* InEOSSubsystem, const int ControllerIndex, bool bShowOnlineOnly, bool bShowSkipButton, const FOnLoginUIClosedDelegate& Delegate)
diff --git a/Engine/Plugins/Online/OnlineSubsystemEOS/Source/OnlineSubsystemEOS/Private/IOS/IOSEOSHelpers.h b/Engine/Plugins/Online/OnlineSubsystemEOS/Source/OnlineSubsystemEOS/Private/IOS/IOSEOSHelpers.h
--- a/Engine/Plugins/Online/OnlineSubsystemEOS/Source/OnlineSubsystemEOS/Private/IOS/IOSEOSHelpers.h
+++ b/Engine/Plugins/Online/OnlineSubsystemEOS/Source/OnlineSubsystemEOS/Private/IOS/IOSEOSHelpers.h
@@ -16,7 +16,7 @@ class FIOSEOSHelpers : public FEOSHelpers
 public:
 	virtual ~FIOSEOSHelpers();
 
-	virtual void PlatformAuthCredentials(EOS_Auth_Credentials &Credentials) override;
+	virtual void GetSystemAuthCredentialsOptions(void*& OutCredentialsOptions) override;
 	virtual void PlatformTriggerLoginUI(FOnlineSubsystemEOS* EOSSubsystem, const int ControllerIndex, bool bShowOnlineOnly, bool bShowSkipButton, const FOnLoginUIClosedDelegate& Delegate) override;
 
 private:
diff --git a/Engine/Plugins/Online/OnlineSubsystemEOS/Source/OnlineSubsystemEOS/Private/UserManagerEOS.cpp b/Engine/Plugins/Online/OnlineSubsystemEOS/Source/OnlineSubsystemEOS/Private/UserManagerEOS.cpp
--- a/Engine/Plugins/Online/OnlineSubsystemEOS/Source/OnlineSubsystemEOS/Private/UserManagerEOS.cpp
+++ b/Engine/Plugins/Online/OnlineSubsystemEOS/Source/OnlineSubsystemEOS/Private/UserManagerEOS.cpp
@@ -129,6 +129,65 @@ static inline EOS_EExternalCredentialType ToEOS_EExternalCredentialType(FName OS
 	return EOS_EExternalCredentialType::EOS_ECT_OPENID_ACCESS_TOKEN;
 }
 
+bool ToEOS_ELoginCredentialType(const FString& InTypeStr, EOS_ELoginCredentialType& OutType)
+{
+	if (InTypeStr == TEXT("exchangecode"))
+	{
+		// This is how the Epic launcher will pass credentials to you
+		OutType = EOS_ELoginCredentialType::EOS_LCT_ExchangeCode;
+	}
+	else if (InTypeStr == TEXT("developer"))
+	{
+		// This is auth via the EOS auth tool
+		OutType = EOS_ELoginCredentialType::EOS_LCT_Developer;
+	}
+	else if (InTypeStr == TEXT("password"))
+	{
+		// This is using a direct username / password. Restricted and not generally available.
+		OutType = EOS_ELoginCredentialType::EOS_LCT_Password;
+	}
+	else if (InTypeStr == TEXT("accountportal"))
+	{
+		// This is auth via the EOS Account Portal
+		OutType = EOS_ELoginCredentialType::EOS_LCT_AccountPortal;
+	}
+	else if (InTypeStr == TEXT("persistentauth"))
+	{
+		// Use locally stored token managed by EOSSDK keyring to attempt login.
+		OutType = EOS_ELoginCredentialType::EOS_LCT_PersistentAuth;
+	}
+	else if (InTypeStr == TEXT("externalauth"))
+	{
+		// Use external auth token to attempt login.
+		OutType = EOS_ELoginCredentialType::EOS_LCT_ExternalAuth;
+	}
+	else
+	{
+		UE_LOG_ONLINE(Warning, TEXT("Incorrect auth parameters (Type=%s)"), *InTypeStr);
+
+		return false;
+	}
+
+	return true;
+}
+
+IOnlineSubsystem* GetPlatformOSS()
+{
+	IOnlineSubsystem* PlatformOSS = IOnlineSubsystem::GetByPlatform();
+	if (PlatformOSS == nullptr)
+#if !PLATFORM_DESKTOP
+	{
+		UE_LOG_ONLINE(Error, TEXT("GetPlatformOSS() failed due to no platform OSS being configured"));
+	}
+#else
+	{
+		// Attempt to load Steam before treating it as an error
+		PlatformOSS = IOnlineSubsystem::Get(STEAM_SUBSYSTEM);
+	}
+#endif
+	return PlatformOSS;
+}
+
 namespace {
 
 EOS_EAuthScopeFlags GetAuthScopeFlags()
@@ -264,23 +323,6 @@ void FUserManagerEOS::LoginStatusChanged(const EOS_Auth_LoginStatusChangedCallba
 	}
 }
 
-IOnlineSubsystem* FUserManagerEOS::GetPlatformOSS() const
-{
-	IOnlineSubsystem* PlatformOSS = IOnlineSubsystem::GetByPlatform();
-	if (PlatformOSS == nullptr)
-#if !PLATFORM_DESKTOP
-	{
-		UE_LOG_ONLINE(Error, TEXT("GetPlatformOSS() failed due to no platform OSS being configured"));
-	}
-#else
-	{
-		// Attempt to load Steam before treating it as an error
-		PlatformOSS = IOnlineSubsystem::Get(STEAM_SUBSYSTEM);
-	}
-#endif
-	return PlatformOSS;
-}
-
 void FUserManagerEOS::GetPlatformAuthToken(int32 LocalUserNum, const FOnGetLinkedAccountAuthTokenCompleteDelegate& Delegate) const
 {
 	IOnlineSubsystem* PlatformOSS = GetPlatformOSS();
@@ -335,96 +377,6 @@ typedef TEOSCallback<EOS_Auth_OnLoginCallback, EOS_Auth_LoginCallbackInfo, FUser
 typedef TEOSCallback<EOS_Connect_OnLoginCallback, EOS_Connect_LoginCallbackInfo, FUserManagerEOS> FConnectLoginCallback;
 typedef TEOSCallback<EOS_Auth_OnDeletePersistentAuthCallback, EOS_Auth_DeletePersistentAuthCallbackInfo, FUserManagerEOS> FDeletePersistentAuthCallback;
 
-// Chose arbitrarily since the SDK doesn't define it
-#define EOS_MAX_TOKEN_SIZE 4096
-
-struct FAuthCredentials :
-	public EOS_Auth_Credentials
-{
-	FAuthCredentials() :
-		EOS_Auth_Credentials()
-	{
-		ApiVersion = 3;
-		UE_EOS_CHECK_API_MISMATCH(EOS_AUTH_CREDENTIALS_API_LATEST, 3);
-		Id = IdAnsi;
-		Token = TokenAnsi;
-
-		FMemory::Memset(IdAnsi, 0, sizeof(IdAnsi));
-		FMemory::Memset(TokenAnsi, 0, sizeof(TokenAnsi));
-	}
-
-	FAuthCredentials(const FAuthCredentials& Other)
-	{
-		ApiVersion = Other.ApiVersion;
-		Id = IdAnsi;
-		Token = TokenAnsi;
-		Type = Other.Type;
-		SystemAuthCredentialsOptions = Other.SystemAuthCredentialsOptions;
-		ExternalType = Other.ExternalType;
-
-		FCStringAnsi::Strncpy(IdAnsi, Other.IdAnsi, EOS_OSS_STRING_BUFFER_LENGTH);
-		FCStringAnsi::Strncpy(TokenAnsi, Other.TokenAnsi, EOS_MAX_TOKEN_SIZE);
-	}
-
-	FAuthCredentials& operator=(FAuthCredentials& Other)
-	{
-		ApiVersion = Other.ApiVersion;
-		Type = Other.Type;
-		SystemAuthCredentialsOptions = Other.SystemAuthCredentialsOptions;
-		ExternalType = Other.ExternalType;
-
-		FCStringAnsi::Strncpy(IdAnsi, Other.IdAnsi, EOS_OSS_STRING_BUFFER_LENGTH);
-		FCStringAnsi::Strncpy(TokenAnsi, Other.TokenAnsi, EOS_MAX_TOKEN_SIZE);
-
-		return *this;
-	}
-
-	FAuthCredentials(EOS_EExternalCredentialType InExternalType, const FExternalAuthToken& AuthToken) :
-		EOS_Auth_Credentials()
-	{
-		if (AuthToken.HasTokenData())
-		{
-			Init(InExternalType, AuthToken.TokenData);
-		}
-		else if (AuthToken.HasTokenString())
-		{
-			Init(InExternalType, AuthToken.TokenString);
-		}
-		else
-		{
-			UE_LOG_ONLINE(Error, TEXT("FAuthCredentials object cannot be constructed with invalid FExternalAuthToken parameter"));
-		}
-	}
-
-	void Init(EOS_EExternalCredentialType InExternalType, const FString& InTokenString)
-	{
-		ApiVersion = 3;
-		UE_EOS_CHECK_API_MISMATCH(EOS_AUTH_CREDENTIALS_API_LATEST, 3);
-		Type = EOS_ELoginCredentialType::EOS_LCT_ExternalAuth;
-		ExternalType = InExternalType;
-		Id = nullptr;
-		Token = TokenAnsi;
-
-		FCStringAnsi::Strncpy(TokenAnsi, TCHAR_TO_UTF8(*InTokenString), InTokenString.Len()+1);
-	}
-
-	void Init(EOS_EExternalCredentialType InExternalType, const TArray<uint8>& InToken)
-	{
-		ApiVersion = 3;
-		UE_EOS_CHECK_API_MISMATCH(EOS_AUTH_CREDENTIALS_API_LATEST, 3);
-		Type = EOS_ELoginCredentialType::EOS_LCT_ExternalAuth;
-		ExternalType = InExternalType;
-		Id = nullptr;
-		Token = TokenAnsi;
-
-		uint32_t InOutBufferLength = EOS_MAX_TOKEN_SIZE;
-		EOS_ByteArray_ToString(InToken.GetData(), InToken.Num(), TokenAnsi, &InOutBufferLength);
-	}
-
-	char IdAnsi[EOS_OSS_STRING_BUFFER_LENGTH];
-	char TokenAnsi[EOS_MAX_TOKEN_SIZE];
-};
-
 bool FUserManagerEOS::Login(int32 LocalUserNum, const FOnlineAccountCredentials& AccountCredentials)
 {
 	LocalUserNumToLastLoginCredentials.Emplace(LocalUserNum, MakeShared<FOnlineAccountCredentials>(AccountCredentials));
@@ -461,66 +413,99 @@ bool FUserManagerEOS::Login(int32 LocalUserNum, const FOnlineAccountCredentials&
 		return true;
 	}
 
-	EOS_Auth_LoginOptions LoginOptions = { };
-	LoginOptions.ApiVersion = 2;
-	UE_EOS_CHECK_API_MISMATCH(EOS_AUTH_LOGIN_API_LATEST, 2);
+	CallEOSAuthLogin(LocalUserNum, AccountCredentials);
 
-	LoginOptions.ScopeFlags = GetAuthScopeFlags();
+	return true;
+}
 
-	FPlatformEOSHelpersPtr EOSHelpers = EOSSubsystem->GetEOSHelpers();
+// Chose arbitrarily since the SDK doesn't define it
+#define EOS_MAX_TOKEN_SIZE 4096
 
-	FAuthCredentials Credentials;
-	LoginOptions.Credentials = &Credentials;
-	EOSHelpers->PlatformAuthCredentials(Credentials);
+FString ToHexString(const TArray<uint8>& InToken)
+{
+	char TokenAnsi[EOS_MAX_TOKEN_SIZE];
+	uint32_t InOutBufferLength = EOS_MAX_TOKEN_SIZE;
+	EOS_ByteArray_ToString(InToken.GetData(), InToken.Num(), TokenAnsi, &InOutBufferLength);
 
-	bool bIsPersistentLogin = false;
+	return FString(TokenAnsi);
+}
 
-	if (AccountCredentials.Type == TEXT("exchangecode"))
-	{
-		// This is how the Epic launcher will pass credentials to you
-		Credentials.Type = EOS_ELoginCredentialType::EOS_LCT_ExchangeCode;
-		Credentials.Id = nullptr;
-		FCStringAnsi::Strncpy(Credentials.TokenAnsi, TCHAR_TO_UTF8(*AccountCredentials.Token), EOS_MAX_TOKEN_SIZE);
-	}
-	else if (AccountCredentials.Type == TEXT("developer"))
-	{
-		// This is auth via the EOS auth tool
-		Credentials.Type = EOS_ELoginCredentialType::EOS_LCT_Developer;
-		FCStringAnsi::Strncpy(Credentials.IdAnsi, TCHAR_TO_UTF8(*AccountCredentials.Id), EOS_OSS_STRING_BUFFER_LENGTH);
-		FCStringAnsi::Strncpy(Credentials.TokenAnsi, TCHAR_TO_UTF8(*AccountCredentials.Token), EOS_MAX_TOKEN_SIZE);
-	}
-	else if (AccountCredentials.Type == TEXT("password"))
-	{
-		// This is using a direct username / password. Restricted and not generally available.
-		Credentials.Type = EOS_ELoginCredentialType::EOS_LCT_Password;
-		FCStringAnsi::Strncpy(Credentials.IdAnsi, TCHAR_TO_UTF8(*AccountCredentials.Id), EOS_OSS_STRING_BUFFER_LENGTH);
-		FCStringAnsi::Strncpy(Credentials.TokenAnsi, TCHAR_TO_UTF8(*AccountCredentials.Token), EOS_MAX_TOKEN_SIZE);
-	}
-	else if (AccountCredentials.Type == TEXT("accountportal"))
-	{
-		// This is auth via the EOS Account Portal
-		Credentials.Type = EOS_ELoginCredentialType::EOS_LCT_AccountPortal;
-		Credentials.Id = nullptr;
-		Credentials.Token = nullptr;
-	}
-	else if (AccountCredentials.Type == TEXT("persistentauth"))
-	{
-		// Use locally stored token managed by EOSSDK keyring to attempt login.
-		Credentials.Type = EOS_ELoginCredentialType::EOS_LCT_PersistentAuth;
-		Credentials.Id = nullptr;
-		Credentials.Token = nullptr;
+void FUserManagerEOS::LoginViaExternalAuth(int32 LocalUserNum)
+{
+	GetPlatformAuthToken(LocalUserNum,
+		FOnGetLinkedAccountAuthTokenCompleteDelegate::CreateLambda([this, WeakThis = AsWeak()](int32 LocalUserNum, bool bWasSuccessful, const FExternalAuthToken& AuthToken)
+		{
+			if (FUserManagerEOSPtr StrongThis = WeakThis.Pin())
+			{
+				if (!bWasSuccessful || !AuthToken.IsValid())
+				{
+					UE_LOG_ONLINE(Warning, TEXT("Unable to Login() user (%d) due to an empty platform auth token"), LocalUserNum);
+					TriggerOnLoginCompleteDelegates(LocalUserNum, false, *FUniqueNetIdEOS::EmptyId(), FString(TEXT("Missing platform auth token")));
+					return;
+				}
 
-		// Store selection of persistent auth.
-		// The persistent auth token is handled by the EOSSDK. On a login failure the persistent token may need to be deleted if it is invalid.
-		bIsPersistentLogin = true;
-	}
-	else
+				check(LocalUserNumToLastLoginCredentials.Contains(LocalUserNum));
+				TSharedRef<FOnlineAccountCredentials> OnlineAccountCredentials = LocalUserNumToLastLoginCredentials[LocalUserNum];
+
+				if (AuthToken.HasTokenData())
+				{
+					OnlineAccountCredentials->Token = ToHexString(AuthToken.TokenData);
+				}
+				else if (AuthToken.HasTokenString())
+				{
+					OnlineAccountCredentials->Token = AuthToken.TokenString;
+				}
+				else
+				{
+					UE_LOG_ONLINE(Error, TEXT("FAuthCredentials object cannot be constructed with invalid FExternalAuthToken parameter"));
+					TriggerOnLoginCompleteDelegates(LocalUserNum, false, *FUniqueNetIdEOS::EmptyId(), FString(TEXT("Invalid platform auth token")));
+					return;
+				}
+
+				CallEOSAuthLogin(LocalUserNum, *OnlineAccountCredentials);
+			}
+		}));
+}
+
+void FUserManagerEOS::CallEOSAuthLogin(int32 LocalUserNum, const FOnlineAccountCredentials& Credentials)
+{
+	// First we construct the EOS Credentials object
+	EOS_Auth_Credentials EOSCredentials = {};
+
+	if (!ToEOS_ELoginCredentialType(Credentials.Type, EOSCredentials.Type))
 	{
 		UE_LOG_ONLINE(Warning, TEXT("Unable to Login() user (%d) due to missing auth parameters"), LocalUserNum);
 		TriggerOnLoginCompleteDelegates(LocalUserNum, false, *FUniqueNetIdEOS::EmptyId(), FString(TEXT("Missing auth parameters")));
-		return false;
+		return;
 	}
 
+	EOSCredentials.ApiVersion = 3;
+	UE_EOS_CHECK_API_MISMATCH(EOS_AUTH_CREDENTIALS_API_LATEST, 3);
+
+	auto IdConverter = StringCast<UTF8CHAR>(*Credentials.Id);
+	EOSCredentials.Id = IdConverter.Length() ? (const char*)IdConverter.Get() : nullptr;
+	auto TokenConverter = StringCast<UTF8CHAR>(*Credentials.Token);
+	EOSCredentials.Token = TokenConverter.Length() ? (const char*)TokenConverter.Get() : nullptr;
+
+	IOnlineSubsystem* PlatformOSS = GetPlatformOSS();
+	EOSCredentials.ExternalType = ToEOS_EExternalCredentialType(PlatformOSS ? PlatformOSS->GetSubsystemName() : EOSSubsystem->GetSubsystemName(), Credentials);
+
+	// We start preparing the Login call
+	EOS_Auth_LoginOptions LoginOptions = { };
+	LoginOptions.ApiVersion = 2;
+	UE_EOS_CHECK_API_MISMATCH(EOS_AUTH_LOGIN_API_LATEST, 2);
+
+	LoginOptions.ScopeFlags = GetAuthScopeFlags();
+
+	LoginOptions.Credentials = &EOSCredentials;
+
+	FPlatformEOSHelpersPtr EOSHelpers = EOSSubsystem->GetEOSHelpers();
+	EOSHelpers->GetSystemAuthCredentialsOptions(EOSCredentials.SystemAuthCredentialsOptions);
+
+	// Store selection of persistent auth.
+	// The persistent auth token is handled by the EOSSDK. On a login failure the persistent token may need to be deleted if it is invalid.
+	const bool bIsPersistentLogin = EOSCredentials.Type == EOS_ELoginCredentialType::EOS_LCT_PersistentAuth;
+
 	FLoginCallback* CallbackObj = new FLoginCallback(AsWeak());
 	CallbackObj->CallbackLambda = [this, LocalUserNum, bIsPersistentLogin](const EOS_Auth_LoginCallbackInfo* Data)
 	{
@@ -529,6 +514,11 @@ bool FUserManagerEOS::Login(int32 LocalUserNum, const FOnlineAccountCredentials&
 			// Continue the login process by getting the product user id for EAS only
 			ConnectLoginEAS(LocalUserNum, Data->LocalUserId);
 		}
+		else if (Data->ResultCode == EOS_EResult::EOS_InvalidUser)
+		{
+			// Link the account
+			LinkEAS(LocalUserNum, Data->ContinuanceToken);
+		}
 		else
 		{
 			auto TriggerLoginFailure = [this, LocalUserNum, LoginResultCode = Data->ResultCode]()
@@ -565,58 +555,9 @@ bool FUserManagerEOS::Login(int32 LocalUserNum, const FOnlineAccountCredentials&
 			}
 		}
 	};
+
 	// Perform the auth call
 	EOS_Auth_Login(EOSSubsystem->AuthHandle, &LoginOptions, (void*)CallbackObj, CallbackObj->GetCallbackPtr());
-	return true;
-}
-
-void FUserManagerEOS::LoginViaExternalAuth(int32 LocalUserNum)
-{
-	GetPlatformAuthToken(LocalUserNum,
-		FOnGetLinkedAccountAuthTokenCompleteDelegate::CreateLambda([this, WeakThis = AsWeak()](int32 LocalUserNum, bool bWasSuccessful, const FExternalAuthToken& AuthToken)
-		{
-			if (FUserManagerEOSPtr StrongThis = WeakThis.Pin())
-			{
-				if (!bWasSuccessful || !AuthToken.IsValid())
-				{
-					UE_LOG_ONLINE(Warning, TEXT("Unable to Login() user (%d) due to an empty platform auth token"), LocalUserNum);
-					TriggerOnLoginCompleteDelegates(LocalUserNum, false, *FUniqueNetIdEOS::EmptyId(), FString(TEXT("Missing platform auth token")));
-					return;
-				}
-
-				EOS_Auth_LoginOptions LoginOptions = { };
-				LoginOptions.ApiVersion = 2;
-				UE_EOS_CHECK_API_MISMATCH(EOS_AUTH_LOGIN_API_LATEST, 2);
-
-				LoginOptions.ScopeFlags = GetAuthScopeFlags();
-
-				check(LocalUserNumToLastLoginCredentials.Contains(LocalUserNum));
-				FAuthCredentials Credentials(ToEOS_EExternalCredentialType(GetPlatformOSS()->GetSubsystemName(), *LocalUserNumToLastLoginCredentials[LocalUserNum]), AuthToken);
-				LoginOptions.Credentials = &Credentials;
-
-				FLoginCallback* CallbackObj = new FLoginCallback(AsWeak());
-				CallbackObj->CallbackLambda = [this, LocalUserNum](const EOS_Auth_LoginCallbackInfo* Data)
-				{
-					if (Data->ResultCode == EOS_EResult::EOS_Success)
-					{
-						ConnectLoginEAS(LocalUserNum, Data->LocalUserId);
-					}
-					else if (Data->ResultCode == EOS_EResult::EOS_InvalidUser)
-					{
-						// Link the account
-						LinkEAS(LocalUserNum, Data->ContinuanceToken);
-					}
-					else
-					{
-						FString ErrorString = FString::Printf(TEXT("Login(%d) failed with EOS result code (%s)"), LocalUserNum, ANSI_TO_TCHAR(EOS_EResult_ToString(Data->ResultCode)));
-						UE_LOG_ONLINE(Warning, TEXT("%s"), *ErrorString);
-						TriggerOnLoginCompleteDelegates(LocalUserNum, false, *FUniqueNetIdEOS::EmptyId(), ErrorString);
-					}
-				};
-				// Perform the auth call
-				EOS_Auth_Login(EOSSubsystem->AuthHandle, &LoginOptions, (void*)CallbackObj, CallbackObj->GetCallbackPtr());
-			}
-		}));
 }
 
 struct FLinkAccountOptions :
diff --git a/Engine/Plugins/Online/OnlineSubsystemEOS/Source/OnlineSubsystemEOS/Private/UserManagerEOS.h b/Engine/Plugins/Online/OnlineSubsystemEOS/Source/OnlineSubsystemEOS/Private/UserManagerEOS.h
--- a/Engine/Plugins/Online/OnlineSubsystemEOS/Source/OnlineSubsystemEOS/Private/UserManagerEOS.h
+++ b/Engine/Plugins/Online/OnlineSubsystemEOS/Source/OnlineSubsystemEOS/Private/UserManagerEOS.h
@@ -325,6 +325,8 @@ public:
 	int32 GetDefaultLocalUser() const { return DefaultLocalUser; }
 
 private:
+	void CallEOSAuthLogin(int32 LocalUserNum, const FOnlineAccountCredentials& Credentials);
+
 	void RemoveLocalUser(int32 LocalUserNum);
 	void AddLocalUser(int32 LocalUserNum, EOS_EpicAccountId EpicAccountId, EOS_ProductUserId UserId);
 
@@ -341,7 +343,6 @@ private:
 	void UpdatePresence(EOS_EpicAccountId AccountId);
 	void UpdateFriendPresence(const FString& FriendId, FOnlineUserPresenceRef Presence);
 
-	IOnlineSubsystem* GetPlatformOSS() const;
 	void GetPlatformAuthToken(int32 LocalUserNum, const FOnGetLinkedAccountAuthTokenCompleteDelegate& Delegate) const;
 	FString GetPlatformDisplayName(int32 LocalUserNum) const;
 
