- Fix OSSEOS session search only returning 1 entry in some cases.

diff --git a/Engine/Plugins/Online/OnlineSubsystemEOS/Source/OnlineSubsystemEOS/Private/OnlineSessionEOS.cpp b/Engine/Plugins/Online/OnlineSubsystemEOS/Source/OnlineSubsystemEOS/Private/OnlineSessionEOS.cpp
--- a/Engine/Plugins/Online/OnlineSubsystemEOS/Source/OnlineSubsystemEOS/Private/OnlineSessionEOS.cpp
+++ b/Engine/Plugins/Online/OnlineSubsystemEOS/Source/OnlineSubsystemEOS/Private/OnlineSessionEOS.cpp
@@ -4394,22 +4394,9 @@ void FOnlineSessionEOS::StartLobbySearch(int32 SearchingPlayerNum, EOS_HLobbySea
 					EOS_EResult Result = EOS_LobbySearch_CopySearchResultByIndex(LobbySearchHandle, &CopySearchResultByIndexOptions, &LobbyDetailsHandle);
 					if (Result == EOS_EResult::EOS_Success)
 					{
-						TSharedRef<FLobbyDetailsEOS> LobbyDetails = MakeShared<FLobbyDetailsEOS>(LobbyDetailsHandle);
-
 						UE_LOG_ONLINE_SESSION(Verbose, TEXT("[FOnlineSessionEOS::StartLobbySearch::FLobbySearchFindCallback] LobbySearch_CopySearchResultByIndex was successful."));
-
+						TSharedRef<FLobbyDetailsEOS> LobbyDetails = MakeShared<FLobbyDetailsEOS>(LobbyDetailsHandle);
 						PendingLobbySearchResults.Add(LobbyDetails);
-
-						AddLobbySearchResult(LobbyDetails, SearchSettings, [this, LobbyDetails, CompletionDelegate, SearchingPlayerNum, SearchSettings](bool bWasSuccessful)
-						{
-							PendingLobbySearchResults.Remove(LobbyDetails);
-
-							if (PendingLobbySearchResults.IsEmpty())
-							{
-								// If we fail to copy the lobby data, we won't add a new search result, so we'll return an empty one
-								CompletionDelegate.ExecuteIfBound(SearchingPlayerNum, bWasSuccessful, bWasSuccessful ? SearchSettings->SearchResults.Last() : FOnlineSessionSearchResult());
-							}
-						});
 					}
 					else
 					{
@@ -4417,6 +4404,22 @@ void FOnlineSessionEOS::StartLobbySearch(int32 SearchingPlayerNum, EOS_HLobbySea
 					}
 				}
 
+				// Make a copy to iterate over, as the AddLobbySearchResult delegate removes entries.
+				const TArray<TSharedRef<FLobbyDetailsEOS>> PendingLobbySearchResultsCopy = PendingLobbySearchResults;
+				for (const TSharedRef<FLobbyDetailsEOS>& LobbyDetails : PendingLobbySearchResultsCopy)
+				{
+					AddLobbySearchResult(LobbyDetails, SearchSettings, [this, LobbyDetails, CompletionDelegate, SearchingPlayerNum, SearchSettings](bool bWasSuccessful)
+					{
+						PendingLobbySearchResults.Remove(LobbyDetails);
+
+						if (PendingLobbySearchResults.IsEmpty())
+						{
+							// If we fail to copy the lobby data, we won't add a new search result, so we'll return an empty one
+							CompletionDelegate.ExecuteIfBound(SearchingPlayerNum, bWasSuccessful, bWasSuccessful ? SearchSettings->SearchResults.Last() : FOnlineSessionSearchResult());
+						}
+					});
+				}
+
 				if (PendingLobbySearchResults.IsEmpty())
 				{
 					CompletionDelegate.ExecuteIfBound(SearchingPlayerNum, true, SearchSettings->SearchResults.Last());
