Fallback to not using sandbox ID for EOS artifacts.

diff --git a/Engine/Plugins/Online/OnlineSubsystemEOS/Source/OnlineSubsystemEOS/Private/EOSSettings.cpp b/Engine/Plugins/Online/OnlineSubsystemEOS/Source/OnlineSubsystemEOS/Private/EOSSettings.cpp
--- a/Engine/Plugins/Online/OnlineSubsystemEOS/Source/OnlineSubsystemEOS/Private/EOSSettings.cpp
+++ b/Engine/Plugins/Online/OnlineSubsystemEOS/Source/OnlineSubsystemEOS/Private/EOSSettings.cpp
@@ -212,9 +212,12 @@ bool UEOSSettings::GetSelectedArtifactSettings(FEOSArtifactSettings& OutSettings
 	FString SandboxId;
 	if (FParse::Value(FCommandLine::Get(), TEXT("EpicSandboxId="), SandboxId))
 	{
-		const bool bSuccess = GetArtifactSettings(ArtifactName, SandboxId, OutSettings);
-		UE_CLOG_ONLINE(!bSuccess, Error, TEXT("UEOSSettings::GetSelectedArtifactSettings() ArtifactName=[%s] SandboxId=[%s] no settings found."), *ArtifactName, *SandboxId);
-		return bSuccess;
+		if (GetArtifactSettings(ArtifactName, SandboxId, OutSettings))
+		{
+			return true;
+		}
+
+		UE_LOG_ONLINE(Warning, TEXT("UEOSSettings::GetSelectedArtifactSettings() ArtifactName=[%s] SandboxId=[%s] no settings found."), *ArtifactName, *SandboxId);
 	}
 
 	// Fall back on just matching the Artifact name. This assumes non-EGS and only one settings entry per ArtifactName in config.
