Fix for external auth in OSS EOS after recent refactor

diff --git a/Engine/Plugins/Online/OnlineSubsystemEOS/Source/OnlineSubsystemEOS/Private/UserManagerEOS.cpp b/Engine/Plugins/Online/OnlineSubsystemEOS/Source/OnlineSubsystemEOS/Private/UserManagerEOS.cpp
--- a/Engine/Plugins/Online/OnlineSubsystemEOS/Source/OnlineSubsystemEOS/Private/UserManagerEOS.cpp
+++ b/Engine/Plugins/Online/OnlineSubsystemEOS/Source/OnlineSubsystemEOS/Private/UserManagerEOS.cpp
@@ -446,6 +446,7 @@ void FUserManagerEOS::LoginViaExternalAuth(int32 LocalUserNum)
 
 				check(LocalUserNumToLastLoginCredentials.Contains(LocalUserNum));
 				TSharedRef<FOnlineAccountCredentials> OnlineAccountCredentials = LocalUserNumToLastLoginCredentials[LocalUserNum];
+				OnlineAccountCredentials->Type = FString(TEXT("externalauth"));
 
 				if (AuthToken.HasTokenData())
 				{
