- Move bUseSessionPresenceAttribute warning so it only fires when someone uses the sessions API.

diff --git a/Engine/Plugins/Online/OnlineSubsystemEOS/Source/OnlineSubsystemEOS/Private/OnlineSessionEOS.cpp b/Engine/Plugins/Online/OnlineSubsystemEOS/Source/OnlineSubsystemEOS/Private/OnlineSessionEOS.cpp
--- a/Engine/Plugins/Online/OnlineSubsystemEOS/Source/OnlineSubsystemEOS/Private/OnlineSessionEOS.cpp
+++ b/Engine/Plugins/Online/OnlineSubsystemEOS/Source/OnlineSubsystemEOS/Private/OnlineSessionEOS.cpp
@@ -504,10 +504,7 @@ void FOnlineSessionEOS::Init()
 	bIsUsingP2PSockets = false;
 
 	// Presence usage used to be guessed from other session parameters like permission level, with this new logic we have a reliable way of transmitting that information, in the form of a session attribute
-	if (!GConfig->GetBool(TEXT("OnlineSubsystemEOS"), TEXT("bUseSessionPresenceAttribute"), bUsePresenceAttribute, GEngineIni))
-	{
-		UE_LOG_ONLINE_SESSION(Warning, TEXT("Upgrade note: OSSEOS is updating how Session and Lobby presence information is tracked to match other OSS implementations. For more information, search for bUsePresenceAttribute in the release notes."));
-	}
+	GConfig->GetBool(TEXT("OnlineSubsystemEOS"), TEXT("bUseSessionPresenceAttribute"), bUsePresenceAttribute, GEngineIni);
 
  	if (!GConfig->GetBool(TEXT("/Script/OnlineSubsystemEOS.NetDriverEOS"), TEXT("bIsUsingP2PSockets"), bIsUsingP2PSockets, GEngineIni))
 	{
@@ -1295,6 +1292,8 @@ uint32 FOnlineSessionEOS::CreateEOSSession(int32 HostingPlayerNum, FNamedOnlineS
 
 	EOS_HSessionModification SessionModHandle = nullptr;
 
+	TEMP_LogPresenceAttribWarning();
+
 	EOS_Bool bPresenceEnabled = EOS_FALSE;
 	if (bUsePresenceAttribute)
 	{
@@ -2096,6 +2095,8 @@ void FOnlineSessionEOS::CopySearchResult(EOS_HSessionDetails SessionHandle, EOS_
 	OutSession.SessionSettings.bAllowJoinInProgress = SessionInfo->Settings->bAllowJoinInProgress == EOS_TRUE;
 	OutSession.SessionSettings.bAllowInvites = SessionInfo->Settings->bInvitesAllowed == EOS_TRUE;
 
+	TEMP_LogPresenceAttribWarning();
+
 	if (!bUsePresenceAttribute)
 	{
 		OutSession.SessionSettings.bUsesPresence = true;
@@ -2138,6 +2139,8 @@ void FOnlineSessionEOS::CopyAttributes(EOS_HSessionDetails SessionHandle, FOnlin
 		if (ResultCode == EOS_EResult::EOS_Success)
 		{
 			FString Key = Attribute->Data->Key;
+
+			TEMP_LogPresenceAttribWarning();
 			if (bUsePresenceAttribute && Key == USES_PRESENCE_ATTRIBUTE_KEY.ToString())
 			{
 				OutSession.SessionSettings.bUsesPresence = Attribute->Data->Value.AsBool == EOS_TRUE;
@@ -3721,6 +3724,7 @@ uint32 FOnlineSessionEOS::CreateLobbySession(int32 HostingPlayerNum, FNamedOnlin
 	bool bUseHostMigration = true;
 	Session->SessionSettings.Get(SETTING_HOST_MIGRATION, bUseHostMigration);
 
+	TEMP_LogPresenceAttribWarning();
 	if (bUsePresenceAttribute)
 	{
 		Session->SessionSettings.Settings.Add(USES_PRESENCE_ATTRIBUTE_KEY, FOnlineSessionSetting(Session->SessionSettings.bUsesPresence, EOnlineDataAdvertisementType::ViaOnlineService));
@@ -4497,6 +4501,7 @@ void FOnlineSessionEOS::CopyLobbyData(const TSharedRef<FLobbyDetailsEOS>& LobbyD
 	OutSession.SessionSettings.bUseLobbiesVoiceChatIfAvailable = LobbyDetailsInfo->bRTCRoomEnabled == EOS_TRUE;
 #endif
 
+	TEMP_LogPresenceAttribWarning();
 	switch (LobbyDetailsInfo->PermissionLevel)
 	{
 	case EOS_ELobbyPermissionLevel::EOS_LPL_PUBLICADVERTISED:
@@ -4612,6 +4617,7 @@ void FOnlineSessionEOS::CopyLobbyAttributes(const TSharedRef<FLobbyDetailsEOS>&
 		EOS_EResult ResultCode = EOS_LobbyDetails_CopyAttributeByIndex(LobbyDetails->LobbyDetailsHandle, &AttrOptions, &Attribute);
 		if (ResultCode == EOS_EResult::EOS_Success)
 		{
+			TEMP_LogPresenceAttribWarning();
 			FString Key = UTF8_TO_TCHAR(Attribute->Data->Key);
 			if (bUsePresenceAttribute && Key == USES_PRESENCE_ATTRIBUTE_KEY.ToString())
 			{
@@ -4737,4 +4743,11 @@ void FOnlineSessionEOS::CopyLobbyMemberAttributes(const FLobbyDetailsEOS& LobbyD
 	}
 }
 
+void FOnlineSessionEOS::TEMP_LogPresenceAttribWarning()
+{
+	static bool bLogged = false;
+	UE_CLOG_ONLINE_SESSION(!bLogged && !bUsePresenceAttribute, Warning, TEXT("Upgrade note: OSSEOS is updating how Session and Lobby presence information is tracked to match other OSS implementations. For more information, search for bUseSessionPresenceAttribute in the release notes."));
+	bLogged = true;
+}
+
 #endif // WITH_EOS_SDK
diff --git a/Engine/Plugins/Online/OnlineSubsystemEOS/Source/OnlineSubsystemEOS/Private/OnlineSessionEOS.h b/Engine/Plugins/Online/OnlineSubsystemEOS/Source/OnlineSubsystemEOS/Private/OnlineSessionEOS.h
--- a/Engine/Plugins/Online/OnlineSubsystemEOS/Source/OnlineSubsystemEOS/Private/OnlineSessionEOS.h
+++ b/Engine/Plugins/Online/OnlineSubsystemEOS/Source/OnlineSubsystemEOS/Private/OnlineSessionEOS.h
@@ -345,6 +345,9 @@ private:
 	bool bIsDedicatedServer;
 	bool bIsUsingP2PSockets;
 	bool bUsePresenceAttribute = false;
+
+	// Log a warning if bUsePresenceAttribute is false, prompting people to upgrade.
+	void TEMP_LogPresenceAttribWarning();
 };
 
 typedef TSharedPtr<FOnlineSessionEOS, ESPMode::ThreadSafe> FOnlineSessionEOSPtr;
