Refactor for EOS Friend status update handler method, to trigger the appropriate corresponding UE delegates and keep a coherent state.

diff --git a/Engine/Plugins/Online/OnlineSubsystemEOS/Source/OnlineSubsystemEOS/Private/UserManagerEOS.cpp b/Engine/Plugins/Online/OnlineSubsystemEOS/Source/OnlineSubsystemEOS/Private/UserManagerEOS.cpp
--- a/Engine/Plugins/Online/OnlineSubsystemEOS/Source/OnlineSubsystemEOS/Private/UserManagerEOS.cpp
+++ b/Engine/Plugins/Online/OnlineSubsystemEOS/Source/OnlineSubsystemEOS/Private/UserManagerEOS.cpp
@@ -1716,43 +1716,88 @@ void FUserManagerEOS::FriendStatusChanged(const EOS_Friends_OnFriendsUpdateInfo*
 	{
 		int32 LocalUserNum = AccountIdToUserNumMap[Data->LocalUserId];
 		FUniqueNetIdEOSPtr LocalEOSID = UserNumToNetIdMap[LocalUserNum];
-		// If we don't know them yet, or if they just became a friend, then add them to kick off the reads
-		if (!AccountIdToStringMap.Contains(Data->TargetUserId) || Data->CurrentStatus == EOS_EFriendsStatus::EOS_FS_Friends)
+		if (!AccountIdToStringMap.Contains(Data->TargetUserId))
 		{
-			AddFriend(LocalUserNum, Data->TargetUserId);
+			FUniqueNetIdEOSRef FriendNetId = FUniqueNetIdEOSRegistry::FindOrAdd(Data->TargetUserId, nullptr).ToSharedRef();
+			const FString NetId = FriendNetId->ToString();
+			AddRemotePlayer(LocalUserNum, NetId, Data->TargetUserId);
+			// TODO: We need to modify AddRemotePlayer to take a delegate, to execute after all information has been retrieved
 		}
-		// They are in our list now
+
 		FOnlineUserPtr OnlineUser = EpicAccountIdToOnlineUserMap[Data->TargetUserId];
-		FOnlineFriendEOSPtr Friend = LocalUserNumToFriendsListMap[LocalUserNum]->GetByNetIdString(AccountIdToStringMap[Data->TargetUserId]);
 
-		// Figure out which notification to fire. Invite related notifications will only fire in projects with friends management enabled
-		if (Data->CurrentStatus == EOS_EFriendsStatus::EOS_FS_Friends)
+		switch (Data->CurrentStatus)
 		{
-			Friend->SetInviteStatus(EInviteStatus::Accepted);
-			TriggerOnInviteAcceptedDelegates(*LocalEOSID, *OnlineUser->GetUserId());
-		}
-		else if (Data->PreviousStatus == EOS_EFriendsStatus::EOS_FS_Friends && Data->CurrentStatus == EOS_EFriendsStatus::EOS_FS_NotFriends)
+		case EOS_EFriendsStatus::EOS_FS_NotFriends: // Invite rejections and friend removal
 		{
-			LocalUserNumToFriendsListMap[LocalUserNum]->Remove(AccountIdToStringMap[Data->TargetUserId], Friend.ToSharedRef());
-			Friend->SetInviteStatus(EInviteStatus::Unknown);
-			TriggerOnFriendRemovedDelegates(*LocalEOSID, *OnlineUser->GetUserId());
+			//User should already be a friend
+			const FString& AccountIdStr = AccountIdToStringMap[Data->TargetUserId];
+			FFriendsListEOSRef& FriendsListRef = LocalUserNumToFriendsListMap[LocalUserNum];
+			FOnlineFriendEOSPtr Friend = FriendsListRef->GetByNetIdString(AccountIdStr);
+			if (Friend.IsValid())
+			{
+				FriendsListRef->Remove(AccountIdStr, Friend.ToSharedRef());
+				Friend->SetInviteStatus(EInviteStatus::Unknown);
+
+				if (Data->PreviousStatus == EOS_EFriendsStatus::EOS_FS_Friends)
+				{
+					TriggerOnFriendRemovedDelegates(*LocalEOSID, *OnlineUser->GetUserId());
+				}
+				else if (Data->PreviousStatus == EOS_EFriendsStatus::EOS_FS_InviteSent || Data->PreviousStatus == EOS_EFriendsStatus::EOS_FS_InviteReceived)
+				{
+					TriggerOnInviteRejectedDelegates(*LocalEOSID, *OnlineUser->GetUserId()); // We don't have an "OnInviteRejected" event only for the local user
+				}
+			}
+			else
+			{
+				UE_LOG_ONLINE_FRIEND(Verbose, TEXT("Friend status notification received for user [%d], but remote user [%s] was not previously registered as a friend"), *AccountIdStr);
+			}
+
+			break;
 		}
-		else if (Data->PreviousStatus < EOS_EFriendsStatus::EOS_FS_Friends && Data->CurrentStatus == EOS_EFriendsStatus::EOS_FS_NotFriends)
+		case EOS_EFriendsStatus::EOS_FS_InviteSent: // Invite sent by local user
 		{
-			LocalUserNumToFriendsListMap[LocalUserNum]->Remove(AccountIdToStringMap[Data->TargetUserId], Friend.ToSharedRef());
-			Friend->SetInviteStatus(EInviteStatus::Unknown);
-			TriggerOnInviteRejectedDelegates(*LocalEOSID, *OnlineUser->GetUserId());
+			// The only supported case is NotFriends->InviteSent, other combinations are not possible
+			check(Data->PreviousStatus == EOS_EFriendsStatus::EOS_FS_NotFriends);
+
+			FOnlineFriendEOSRef FriendRef = AddFriend(LocalUserNum, Data->TargetUserId);
+
+			FriendRef->SetInviteStatus(EInviteStatus::PendingOutbound);
+			TriggerOnOutgoingInviteSentDelegates(LocalUserNum);
+
+			break;
 		}
-		else if (Data->CurrentStatus == EOS_EFriendsStatus::EOS_FS_InviteReceived)
+		case EOS_EFriendsStatus::EOS_FS_InviteReceived: // Invite received by local user
 		{
-			Friend->SetInviteStatus(EInviteStatus::PendingInbound);
+			// The only supported case is NotFriends->InviteReceived, other combinations are not possible
+			check(Data->PreviousStatus == EOS_EFriendsStatus::EOS_FS_NotFriends);
+
+			FOnlineFriendEOSRef FriendRef = AddFriend(LocalUserNum, Data->TargetUserId);
+
+			FriendRef->SetInviteStatus(EInviteStatus::PendingInbound);
 			TriggerOnInviteReceivedDelegates(*LocalEOSID, *OnlineUser->GetUserId());
+
+			break;
+		}
+		case EOS_EFriendsStatus::EOS_FS_Friends: // Accepted invites and friend addition. This logic applies to all three PreviousStatus cases
+		{
+			FOnlineFriendEOSRef FriendRef = AddFriend(LocalUserNum, Data->TargetUserId);
+
+			FriendRef->SetInviteStatus(EInviteStatus::Accepted);
+			TriggerOnInviteAcceptedDelegates(*LocalEOSID, *OnlineUser->GetUserId()); // We don't have an "OnFriendAdded", all relationships exist in the context of invites. We also don't have a separate "OnInviteAccepted" for the local user.
+
+			break;
 		}
+		default:
+			checkNoEntry();
+			UE_LOG_ONLINE_FRIEND(Warning, TEXT("[FUserManagerEOS::FriendStatusChanged] Unsupported status received as CurrentStatus"));
+		}
+
 		TriggerOnFriendsChangeDelegates(LocalUserNum);
 	}
 }
 
-void FUserManagerEOS::AddFriend(int32 LocalUserNum, EOS_EpicAccountId EpicAccountId)
+FOnlineFriendEOSRef FUserManagerEOS::AddFriend(int32 LocalUserNum, EOS_EpicAccountId EpicAccountId)
 {
 	FUniqueNetIdEOSRef FriendNetId = FUniqueNetIdEOSRegistry::FindOrAdd(EpicAccountId, nullptr).ToSharedRef();
 	const FString NetId = FriendNetId->ToString();
@@ -1768,16 +1813,21 @@ void FUserManagerEOS::AddFriend(int32 LocalUserNum, EOS_EpicAccountId EpicAccoun
 	
 	FriendRef->SetInviteStatus(ToEInviteStatus(Status));
 
-	// Add this friend as a remote player (this will grab user info)
-	AddRemotePlayer(LocalUserNum, NetId, EpicAccountId, FriendNetId, FriendRef, FriendRef);
+	// Add this friend as a remote player if we haven't already (this will grab user info)
+	if (!AccountIdToStringMap.Contains(EpicAccountId))
+	{
+		AddRemotePlayer(LocalUserNum, NetId, EpicAccountId, FriendNetId, FriendRef, FriendRef);
+	}
 
 	// Querying the presence of a non-friend would cause an SDK error.
-	// Players that sent/recieved a friend invitation from us still count as "friends", so check
+	// Players that sent/received a friend invitation from us still count as "friends", so check
 	// our friend relationship here.
 	if(Status == EOS_EFriendsStatus::EOS_FS_Friends)
 	{
 		QueryPresence(*FriendNetId, IgnoredPresenceDelegate);
 	}
+
+	return FriendRef;
 }
 
 void FUserManagerEOS::AddRemotePlayer(int32 LocalUserNum, const FString& NetId, EOS_EpicAccountId EpicAccountId)
diff --git a/Engine/Plugins/Online/OnlineSubsystemEOS/Source/OnlineSubsystemEOS/Private/UserManagerEOS.h b/Engine/Plugins/Online/OnlineSubsystemEOS/Source/OnlineSubsystemEOS/Private/UserManagerEOS.h
--- a/Engine/Plugins/Online/OnlineSubsystemEOS/Source/OnlineSubsystemEOS/Private/UserManagerEOS.h
+++ b/Engine/Plugins/Online/OnlineSubsystemEOS/Source/OnlineSubsystemEOS/Private/UserManagerEOS.h
@@ -330,7 +330,7 @@ private:
 	void RemoveLocalUser(int32 LocalUserNum);
 	void AddLocalUser(int32 LocalUserNum, EOS_EpicAccountId EpicAccountId, EOS_ProductUserId UserId);
 
-	void AddFriend(int32 LocalUserNum, EOS_EpicAccountId EpicAccountId);
+	FOnlineFriendEOSRef AddFriend(int32 LocalUserNum, EOS_EpicAccountId EpicAccountId);
 	void AddRemotePlayer(int32 LocalUserNum, const FString& NetId, EOS_EpicAccountId EpicAccountId);
 	void AddRemotePlayer(int32 LocalUserNum, const FString& NetId, EOS_EpicAccountId EpicAccountId, FUniqueNetIdEOSPtr UniqueNetId, FOnlineUserPtr OnlineUser, IAttributeAccessInterfaceRef AttributeRef);
 	void UpdateRemotePlayerProductUserId(EOS_EpicAccountId AccountId, EOS_ProductUserId UserId);
