Added logic to trigger OnQueryUserInfoComplete delegates on OSS EOS

diff --git a/Engine/Plugins/Online/OnlineSubsystemEOS/Source/OnlineSubsystemEOS/Private/UserManagerEOS.cpp b/Engine/Plugins/Online/OnlineSubsystemEOS/Source/OnlineSubsystemEOS/Private/UserManagerEOS.cpp
--- a/Engine/Plugins/Online/OnlineSubsystemEOS/Source/OnlineSubsystemEOS/Private/UserManagerEOS.cpp
+++ b/Engine/Plugins/Online/OnlineSubsystemEOS/Source/OnlineSubsystemEOS/Private/UserManagerEOS.cpp
@@ -211,7 +211,6 @@ EOS_EAuthScopeFlags GetAuthScopeFlags()
 
 /** Delegates that are used for internal calls and are meant to be ignored */
 IOnlinePresence::FOnPresenceTaskCompleteDelegate IgnoredPresenceDelegate;
-IOnlineUser::FOnQueryExternalIdMappingsComplete IgnoredMappingDelegate;
 
 typedef TEOSGlobalCallback<EOS_UI_OnDisplaySettingsUpdatedCallback, EOS_UI_OnDisplaySettingsUpdatedCallbackInfo, FUserManagerEOS> FOnDisplaySettingsUpdatedCallback;
 
@@ -2729,7 +2728,25 @@ bool FUserManagerEOS::QueryUserInfo(int32 LocalUserNum, const TArray<FUniqueNetI
 	}
 
 	const FUniqueNetIdEOSPtr LocalId = GetLocalUniqueNetIdEOS(LocalUserNum);
-	QueryExternalIdMappings(*LocalId, FExternalIdQueryOptions(), UserEasIdsNeedingExternalMappings, IgnoredMappingDelegate);
+	QueryExternalIdMappings(*LocalId, FExternalIdQueryOptions(), UserEasIdsNeedingExternalMappings, FOnQueryExternalIdMappingsComplete::CreateLambda([this, WeakThis = AsWeak(), LocalUserNum](bool bWasSuccessful, const FUniqueNetId& UserId, const FExternalIdQueryOptions& QueryOptions, const TArray<FString>& ExternalIds, const FString& ErrorStr)
+		{
+			if (FUserManagerEOSPtr StrongThis = WeakThis.Pin())
+			{
+				TArray<FUniqueNetIdPtr> ExternalIdPtrs;
+				GetExternalIdMappings(QueryOptions, ExternalIds, ExternalIdPtrs);
+
+				TArray<FUniqueNetIdRef> ExternalIdRefs;
+				for (const FUniqueNetIdPtr& ExternalIdPtr : ExternalIdPtrs)
+				{
+					if (ExternalIdPtr.IsValid())
+					{
+						ExternalIdRefs.Add(ExternalIdPtr.ToSharedRef());
+					}
+				}
+
+				TriggerOnQueryUserInfoCompleteDelegates(LocalUserNum, bWasSuccessful, ExternalIdRefs, ErrorStr);
+			}
+		}));
 	
 	return true;
 }
