UserManagerEOS internal data structures refactor (Part 1)
- Started removal of redundant maps, replaced with lookup methods. Functionality should remain the same, but with only a fraction of the internal data structures.
- Added and reorganised methods to ensure coverage of all id method conversions

diff --git a/Engine/Plugins/Online/OnlineSubsystemEOS/Source/OnlineSubsystemEOS/Private/UserManagerEOS.cpp b/Engine/Plugins/Online/OnlineSubsystemEOS/Source/OnlineSubsystemEOS/Private/UserManagerEOS.cpp
--- a/Engine/Plugins/Online/OnlineSubsystemEOS/Source/OnlineSubsystemEOS/Private/UserManagerEOS.cpp
+++ b/Engine/Plugins/Online/OnlineSubsystemEOS/Source/OnlineSubsystemEOS/Private/UserManagerEOS.cpp
@@ -217,7 +217,7 @@ typedef TEOSGlobalCallback<EOS_UI_OnDisplaySettingsUpdatedCallback, EOS_UI_OnDis
 FUserManagerEOS::FUserManagerEOS(FOnlineSubsystemEOS* InSubsystem)
 	: TSharedFromThis<FUserManagerEOS, ESPMode::ThreadSafe>()
 	, EOSSubsystem(InSubsystem)
-	, DefaultLocalUser(-1)
+	, DefaultLocalUser(INVALID_LOCAL_USER)
 	, LoginNotificationId(0)
 	, LoginNotificationCallback(nullptr)
 	, FriendsNotificationId(0)
@@ -274,16 +274,16 @@ void FUserManagerEOS::LoginStatusChanged(const EOS_Auth_LoginStatusChangedCallba
 {
 	if (Data->CurrentStatus == EOS_ELoginStatus::EOS_LS_NotLoggedIn)
 	{
-		if (AccountIdToUserNumMap.Contains(Data->LocalUserId))
+		if (IsLocalUser(Data->LocalUserId))
 		{
-			int32 LocalUserNum = AccountIdToUserNumMap[Data->LocalUserId];
-			FUniqueNetIdEOSPtr UserNetId = UserNumToNetIdMap[LocalUserNum];
+			const int32 LocalUserNum = GetLocalUserNumFromEpicAccountId(Data->LocalUserId);
+			const FUniqueNetIdEOSPtr UserNetId = LocalUserUniqueNetIds[LocalUserNum];
 			TriggerOnLoginStatusChangedDelegates(LocalUserNum, ELoginStatus::LoggedIn, ELoginStatus::NotLoggedIn, *UserNetId);
 			// Need to remove the local user
 			RemoveLocalUser(LocalUserNum);
 
 			// Clean up user based notifies if we have no logged in users
-			if (UserNumToNetIdMap.Num() == 0)
+			if (LocalUserUniqueNetIds.Num() == 0)
 			{
 				if (LoginNotificationId > 0)
 				{
@@ -753,7 +753,8 @@ bool FUserManagerEOS::ConnectLoginEAS(int32 LocalUserNum, EOS_EpicAccountId Acco
 
 void FUserManagerEOS::RefreshConnectLogin(int32 LocalUserNum)
 {
-	if (!UserNumToAccountIdMap.Contains(LocalUserNum))
+	const EOS_EpicAccountId AccountId = GetLocalEpicAccountId(LocalUserNum);
+	if (!AccountId)
 	{
 		UE_LOG_ONLINE(Error, TEXT("Can't refresh ConnectLogin(%d) since (%d) is not logged in"), LocalUserNum, LocalUserNum);
 		return;
@@ -762,7 +763,6 @@ void FUserManagerEOS::RefreshConnectLogin(int32 LocalUserNum)
 	const FEOSSettings Settings = UEOSSettings::GetSettings();
 	if (Settings.bUseEAS)
 	{
-		EOS_EpicAccountId AccountId = UserNumToAccountIdMap[LocalUserNum];
 		EOS_Auth_Token* AuthToken = nullptr;
 		EOS_Auth_CopyUserAuthTokenOptions CopyOptions = { };
 		CopyOptions.ApiVersion = 1;
@@ -772,8 +772,8 @@ void FUserManagerEOS::RefreshConnectLogin(int32 LocalUserNum)
 		if (CopyResult == EOS_EResult::EOS_Success)
 		{
 			// We update the auth token cached in the user account, along with the user information
-			const FUniqueNetIdEOSPtr UniqueNetId = UserNumToNetIdMap.FindChecked(LocalUserNum);
-			const FUserOnlineAccountEOSRef UserAccountRef = StringToUserAccountMap.FindChecked(UniqueNetId->ToString());
+			const FUniqueNetIdEOSRef UniqueNetId = LocalUserUniqueNetIds[LocalUserNum];
+			const FUserOnlineAccountEOSRef UserAccountRef = UniqueNetIdToUserAccountMap.FindChecked(UniqueNetId);
 			UserAccountRef->SetAuthAttribute(AUTH_ATTR_ID_TOKEN, AuthToken->AccessToken);
 			UpdateUserInfo(UserAccountRef, AccountId, AccountId);
 
@@ -1034,7 +1034,7 @@ bool FUserManagerEOS::AutoLogin(int32 LocalUserNum)
 void FUserManagerEOS::AddLocalUser(int32 LocalUserNum, EOS_EpicAccountId EpicAccountId, EOS_ProductUserId UserId)
 {
 	// Set the default user to the first one that logs in
-	if (DefaultLocalUser == -1)
+	if (DefaultLocalUser == INVALID_LOCAL_USER)
 	{
 		DefaultLocalUser = LocalUserNum;
 	}
@@ -1043,16 +1043,12 @@ void FUserManagerEOS::AddLocalUser(int32 LocalUserNum, EOS_EpicAccountId EpicAcc
 	const FString& NetId = UserNetId->ToString();
 	FUserOnlineAccountEOSRef UserAccountRef(new FUserOnlineAccountEOS(UserNetId));
 
-	UserNumToNetIdMap.Emplace(LocalUserNum, UserNetId);
-	UserNumToAccountIdMap.Emplace(LocalUserNum, EpicAccountId);
-	AccountIdToUserNumMap.Emplace(EpicAccountId, LocalUserNum);
+	LocalUserUniqueNetIds.Insert(LocalUserNum, UserNetId);
+	UniqueNetIdToUserAccountMap.Emplace(UserNetId, UserAccountRef);
 	NetIdStringToOnlineUserMap.Emplace(*NetId, UserAccountRef);
-	StringToUserAccountMap.Emplace(NetId, UserAccountRef);
 	AccountIdToStringMap.Emplace(EpicAccountId, NetId);
 	ProductUserIdToStringMap.Emplace(UserId, *NetId);
 	EpicAccountIdToAttributeAccessMap.Emplace(EpicAccountId, UserAccountRef);
-	UserNumToProductUserIdMap.Emplace(LocalUserNum, UserId);
-	ProductUserIdToUserNumMap.Emplace(UserId, LocalUserNum);
 
 	// Init player lists
 	FFriendsListEOSRef FriendsList = MakeShareable(new FFriendsListEOS(LocalUserNum, UserNetId));
@@ -1123,7 +1119,7 @@ TSharedPtr<FUserOnlineAccount> FUserManagerEOS::GetUserAccount(const FUniqueNetI
 	TSharedPtr<FUserOnlineAccount> Result;
 
 	const FUniqueNetIdEOS& EOSID = FUniqueNetIdEOS::Cast(UserId);
-	const FUserOnlineAccountEOSRef* FoundUserAccount = StringToUserAccountMap.Find(EOSID.ToString());
+	const FUserOnlineAccountEOSRef* FoundUserAccount = UniqueNetIdToUserAccountMap.Find(EOSID.AsShared());
 	if (FoundUserAccount != nullptr)
 	{
 		return *FoundUserAccount;
@@ -1136,7 +1132,7 @@ TArray<TSharedPtr<FUserOnlineAccount>> FUserManagerEOS::GetAllUserAccounts() con
 {
 	TArray<TSharedPtr<FUserOnlineAccount>> Result;
 
-	for (TMap<FString, FUserOnlineAccountEOSRef>::TConstIterator It(StringToUserAccountMap); It; ++It)
+	for (TUniqueNetIdMap<FUserOnlineAccountEOSRef>::TConstIterator It(UniqueNetIdToUserAccountMap); It; ++It)
 	{
 		Result.Add(It.Value());
 	}
@@ -1148,104 +1144,146 @@ FUniqueNetIdPtr FUserManagerEOS::GetUniquePlayerId(int32 LocalUserNum) const
 	return GetLocalUniqueNetIdEOS(LocalUserNum);
 }
 
-int32 FUserManagerEOS::GetLocalUserNumFromUniqueNetId(const FUniqueNetId& NetId) const
+const FUniqueNetIdEOSPtr FUserManagerEOS::GetLocalUniqueNetIdEOS(int32 LocalUserNum) const
 {
-	const FUniqueNetIdEOS& EosId = FUniqueNetIdEOS::Cast(NetId);
-
-	const EOS_EpicAccountId AccountId = EosId.GetEpicAccountId();
-	if (const int32* UserNum = AccountIdToUserNumMap.Find(AccountId))
+	if (LocalUserUniqueNetIds.IsValidIndex(LocalUserNum))
 	{
-		return *UserNum;
+		return LocalUserUniqueNetIds[LocalUserNum];
 	}
 
-	const EOS_ProductUserId ProductUserId = EosId.GetProductUserId();
-	if (const int32* UserNum = ProductUserIdToUserNumMap.Find(ProductUserId))
+	return nullptr;
+}
+
+const FUniqueNetIdEOSPtr FUserManagerEOS::GetLocalUniqueNetIdEOS() const
+{
+	return GetLocalUniqueNetIdEOS(DefaultLocalUser);
+}
+
+EOS_EpicAccountId FUserManagerEOS::GetLocalEpicAccountId(int32 LocalUserNum) const
+{
+	if (FUniqueNetIdEOSPtr UniqueNetId = GetLocalUniqueNetIdEOS(LocalUserNum))
 	{
-		return *UserNum;
+		return UniqueNetId->GetEpicAccountId();
 	}
 
-	// Use the default user if we can't find the person that they want
-	return DefaultLocalUser;
+	return nullptr;
 }
 
-bool FUserManagerEOS::IsLocalUser(const FUniqueNetId& NetId) const
+EOS_EpicAccountId FUserManagerEOS::GetLocalEpicAccountId() const
 {
-	const FUniqueNetIdEOS& EosId = FUniqueNetIdEOS::Cast(NetId);
-	return AccountIdToUserNumMap.Contains(EosId.GetEpicAccountId()) || ProductUserIdToUserNumMap.Contains(EosId.GetProductUserId());
+	return GetLocalEpicAccountId(DefaultLocalUser);
 }
 
-FUniqueNetIdEOSPtr FUserManagerEOS::GetLocalUniqueNetIdEOS(int32 LocalUserNum) const
+EOS_ProductUserId FUserManagerEOS::GetLocalProductUserId(int32 LocalUserNum) const
 {
-	const FUniqueNetIdEOSPtr* FoundId = UserNumToNetIdMap.Find(LocalUserNum);
-	if (FoundId != nullptr)
+	if (FUniqueNetIdEOSPtr UniqueNetId = GetLocalUniqueNetIdEOS(LocalUserNum))
 	{
-		return *FoundId;
+		return UniqueNetId->GetProductUserId();
 	}
+
 	return nullptr;
 }
 
-FUniqueNetIdEOSPtr FUserManagerEOS::GetLocalUniqueNetIdEOS(EOS_ProductUserId UserId) const
+EOS_ProductUserId FUserManagerEOS::GetLocalProductUserId() const
 {
-	if (ProductUserIdToUserNumMap.Contains(UserId))
+	return GetLocalProductUserId(DefaultLocalUser);
+}
+
+int32 FUserManagerEOS::GetLocalUserNumFromUniqueNetId(const FUniqueNetId& NetId) const
+{
+	const FUniqueNetIdEOS& EosId = FUniqueNetIdEOS::Cast(NetId);
+
+	for (int32 LocalUserNum = 0; LocalUserNum < LocalUserUniqueNetIds.GetMaxIndex(); LocalUserNum++)
 	{
-		return GetLocalUniqueNetIdEOS(ProductUserIdToUserNumMap[UserId]);
+		if (LocalUserUniqueNetIds.IsValidIndex(LocalUserNum))
+		{
+			const FUniqueNetIdEOSRef& Entry = LocalUserUniqueNetIds[LocalUserNum];
+
+			if (*Entry == EosId)
+			{
+				return LocalUserNum;
+			}
+		}
 	}
-	return nullptr;
+
+	return INVALID_LOCAL_USER;
 }
 
-FUniqueNetIdEOSPtr FUserManagerEOS::GetLocalUniqueNetIdEOS(EOS_EpicAccountId AccountId) const
+int32 FUserManagerEOS::GetLocalUserNumFromEpicAccountId(const EOS_EpicAccountId& EpicAccountId) const
 {
-	if (AccountIdToUserNumMap.Contains(AccountId))
+	for (int32 LocalUserNum = 0; LocalUserNum < LocalUserUniqueNetIds.GetMaxIndex(); LocalUserNum++)
 	{
-		return GetLocalUniqueNetIdEOS(AccountIdToUserNumMap[AccountId]);
+		if (LocalUserUniqueNetIds.IsValidIndex(LocalUserNum))
+		{
+			const FUniqueNetIdEOSRef& Entry = LocalUserUniqueNetIds[LocalUserNum];
+
+			if (Entry->GetEpicAccountId() == EpicAccountId)
+			{
+				return LocalUserNum;
+			}
+		}
 	}
-	return nullptr;
+
+	return INVALID_LOCAL_USER;
 }
 
-EOS_EpicAccountId FUserManagerEOS::GetLocalEpicAccountId(int32 LocalUserNum) const
+int32 FUserManagerEOS::GetLocalUserNumFromProductUserId(const EOS_ProductUserId& ProductUserId) const
 {
-	if (UserNumToAccountIdMap.Contains(LocalUserNum))
+	for (int32 LocalUserNum = 0; LocalUserNum < LocalUserUniqueNetIds.GetMaxIndex(); LocalUserNum++)
 	{
-		return UserNumToAccountIdMap[LocalUserNum];
+		if (LocalUserUniqueNetIds.IsValidIndex(LocalUserNum))
+		{
+			const FUniqueNetIdEOSRef& Entry = LocalUserUniqueNetIds[LocalUserNum];
+
+			if (Entry->GetProductUserId() == ProductUserId)
+			{
+				return LocalUserNum;
+			}
+		}
 	}
-	return nullptr;
+
+	return INVALID_LOCAL_USER;
 }
 
-EOS_EpicAccountId FUserManagerEOS::GetLocalEpicAccountId() const
+const FUniqueNetIdEOSPtr FUserManagerEOS::GetLocalUniqueNetIdEOS(const EOS_EpicAccountId& EpicAccountId) const
 {
-	return GetLocalEpicAccountId(DefaultLocalUser);
+	for (const FUniqueNetIdEOSRef& NetId : LocalUserUniqueNetIds)
+	{
+		if (NetId->GetEpicAccountId() == EpicAccountId)
+		{
+			return NetId;
+		}
+	}
+
+	return nullptr;
 }
 
-EOS_ProductUserId FUserManagerEOS::GetLocalProductUserId(int32 LocalUserNum) const
+const FUniqueNetIdEOSPtr FUserManagerEOS::GetLocalUniqueNetIdEOS(const EOS_ProductUserId& ProductUserId) const
 {
-	if (UserNumToProductUserIdMap.Contains(LocalUserNum))
+	for (const FUniqueNetIdEOSRef& NetId : LocalUserUniqueNetIds)
 	{
-		return UserNumToProductUserIdMap[LocalUserNum];
+		if (NetId->GetProductUserId() == ProductUserId)
+		{
+			return NetId;
+		}
 	}
+
 	return nullptr;
 }
 
-EOS_ProductUserId FUserManagerEOS::GetLocalProductUserId() const
+bool FUserManagerEOS::IsLocalUser(const FUniqueNetId& NetId) const
 {
-	return GetLocalProductUserId(DefaultLocalUser);
+	return GetLocalUserNumFromUniqueNetId(NetId) != INVALID_LOCAL_USER;
 }
 
-EOS_EpicAccountId FUserManagerEOS::GetLocalEpicAccountId(EOS_ProductUserId UserId) const
+bool FUserManagerEOS::IsLocalUser(const EOS_EpicAccountId& EpicAccountId) const
 {
-	if (ProductUserIdToUserNumMap.Contains(UserId))
-	{
-		return GetLocalEpicAccountId(ProductUserIdToUserNumMap[UserId]);
-	}
-	return nullptr;
+	return GetLocalUserNumFromEpicAccountId(EpicAccountId) != INVALID_LOCAL_USER;
 }
 
-EOS_ProductUserId FUserManagerEOS::GetLocalProductUserId(EOS_EpicAccountId AccountId) const
+bool FUserManagerEOS::IsLocalUser(const EOS_ProductUserId& ProductUserId) const
 {
-	if (AccountIdToUserNumMap.Contains(AccountId))
-	{
-		return GetLocalProductUserId(AccountIdToUserNumMap[AccountId]);
-	}
-	return nullptr;
+	return GetLocalUserNumFromProductUserId(ProductUserId) != INVALID_LOCAL_USER;
 }
 
 typedef TEOSCallback<EOS_Connect_OnQueryProductUserIdMappingsCallback, EOS_Connect_QueryProductUserIdMappingsCallbackInfo, FUserManagerEOS> FConnectQueryProductUserIdMappingsCallback;
@@ -1362,9 +1400,10 @@ void FUserManagerEOS::ResolveUniqueNetIds(const TArray<EOS_ProductUserId>& Produ
 FOnlineUserPtr FUserManagerEOS::GetLocalOnlineUser(int32 LocalUserNum) const
 {
 	FOnlineUserPtr OnlineUser;
-	if (UserNumToNetIdMap.Contains(LocalUserNum))
+	if (LocalUserUniqueNetIds.IsValidIndex(LocalUserNum))
 	{
-		const FUniqueNetIdEOSPtr NetId = UserNumToNetIdMap.FindRef(LocalUserNum);
+		const FUniqueNetIdEOSRef& NetId = LocalUserUniqueNetIds[LocalUserNum];
+
 		if (NetIdStringToOnlineUserMap.Contains(NetId->ToString()))
 		{
 			OnlineUser = NetIdStringToOnlineUserMap.FindRef(NetId->ToString());
@@ -1403,28 +1442,23 @@ FOnlineUserPtr FUserManagerEOS::GetOnlineUser(EOS_EpicAccountId AccountId) const
 
 void FUserManagerEOS::RemoveLocalUser(int32 LocalUserNum)
 {
-	const FUniqueNetIdEOSPtr* FoundId = UserNumToNetIdMap.Find(LocalUserNum);
-	if (FoundId != nullptr)
+	if (LocalUserUniqueNetIds.IsValidIndex(LocalUserNum))
 	{
-		EOSSubsystem->ReleaseVoiceChatUserInterface(**FoundId);
+		const FUniqueNetIdEOSRef FoundId = LocalUserUniqueNetIds[LocalUserNum];
+
+		EOSSubsystem->ReleaseVoiceChatUserInterface(*FoundId);
 		LocalUserNumToFriendsListMap.Remove(LocalUserNum);
-		const FString& NetId = (*FoundId)->ToString();
-		const EOS_EpicAccountId AccountId = (*FoundId)->GetEpicAccountId();
-		AccountIdToStringMap.Remove(AccountId);
-		AccountIdToUserNumMap.Remove(AccountId);
+		const FString& NetId = FoundId->ToString();
+		AccountIdToStringMap.Remove(FoundId->GetEpicAccountId());
 		NetIdStringToOnlineUserMap.Remove(NetId);
-		StringToUserAccountMap.Remove(NetId);
-		UserNumToNetIdMap.Remove(LocalUserNum);
-		UserNumToAccountIdMap.Remove(LocalUserNum);
-		EOS_ProductUserId UserId = UserNumToProductUserIdMap[LocalUserNum];
-		ProductUserIdToUserNumMap.Remove(UserId);
-		ProductUserIdToStringMap.Remove(UserId);
-		UserNumToProductUserIdMap.Remove(LocalUserNum);
+		LocalUserUniqueNetIds.RemoveAt(LocalUserNum);
+		UniqueNetIdToUserAccountMap.Remove(FoundId);
+		ProductUserIdToStringMap.Remove(FoundId->GetProductUserId());
 	}
 	// Reset this for the next user login
 	if (LocalUserNum == DefaultLocalUser)
 	{
-		DefaultLocalUser = -1;
+		DefaultLocalUser = INVALID_LOCAL_USER;
 	}
 }
 
@@ -1711,10 +1745,10 @@ typedef TEOSCallback<EOS_Friends_OnQueryFriendsCallback, EOS_Friends_QueryFriend
 void FUserManagerEOS::FriendStatusChanged(const EOS_Friends_OnFriendsUpdateInfo* Data)
 {
 	// Get the local user information
-	if (AccountIdToUserNumMap.Contains(Data->LocalUserId))
+	if (IsLocalUser(Data->LocalUserId))
 	{
-		int32 LocalUserNum = AccountIdToUserNumMap[Data->LocalUserId];
-		FUniqueNetIdEOSPtr LocalEOSID = UserNumToNetIdMap[LocalUserNum];
+		const int32 LocalUserNum = GetLocalUserNumFromEpicAccountId(Data->LocalUserId);
+		const FUniqueNetIdEOSRef LocalEOSID = LocalUserUniqueNetIds[LocalUserNum];
 		if (!AccountIdToStringMap.Contains(Data->TargetUserId))
 		{
 			FUniqueNetIdEOSRef FriendNetId = FUniqueNetIdEOSRegistry::FindOrAdd(Data->TargetUserId, nullptr).ToSharedRef();
@@ -1806,7 +1840,7 @@ FOnlineFriendEOSRef FUserManagerEOS::AddFriend(int32 LocalUserNum, EOS_EpicAccou
 	EOS_Friends_GetStatusOptions Options = { };
 	Options.ApiVersion = 1;
 	UE_EOS_CHECK_API_MISMATCH(EOS_FRIENDS_GETSTATUS_API_LATEST, 1);
-	Options.LocalUserId = UserNumToAccountIdMap[LocalUserNum];
+	Options.LocalUserId = GetLocalEpicAccountId(LocalUserNum);
 	Options.TargetUserId = EpicAccountId;
 	EOS_EFriendsStatus Status = EOS_Friends_GetStatus(EOSSubsystem->FriendsHandle, &Options);
 	
@@ -1900,7 +1934,7 @@ void FUserManagerEOS::UpdateRemotePlayerProductUserId(EOS_EpicAccountId EpicAcco
 
 bool FUserManagerEOS::ReadFriendsList(int32 LocalUserNum, const FString& ListName, const FOnReadFriendsListComplete& Delegate)
 {
-	if (!UserNumToNetIdMap.Contains(LocalUserNum))
+	if (!LocalUserUniqueNetIds.IsValidIndex(LocalUserNum))
 	{
 		const FString ErrorStr = FString::Printf(TEXT("Can't ReadFriendsList() for user (%d) since they are not logged in"), LocalUserNum);
 		UE_LOG_ONLINE_FRIEND(Warning, TEXT("%s"), *ErrorStr);
@@ -1925,7 +1959,7 @@ bool FUserManagerEOS::ReadFriendsList(int32 LocalUserNum, const FString& ListNam
 	EOS_Friends_QueryFriendsOptions Options = { };
 	Options.ApiVersion = 1;
 	UE_EOS_CHECK_API_MISMATCH(EOS_FRIENDS_QUERYFRIENDS_API_LATEST, 1);
-	Options.LocalUserId = UserNumToAccountIdMap[LocalUserNum];
+	Options.LocalUserId = GetLocalEpicAccountId(LocalUserNum);
 
 	FReadFriendsCallback* CallbackObj = new FReadFriendsCallback(AsWeak());
 	CallbackObj->CallbackLambda = [this, LocalUserNum, ListName, Delegate](const EOS_Friends_QueryFriendsCallbackInfo* Data)
@@ -1943,7 +1977,7 @@ bool FUserManagerEOS::ReadFriendsList(int32 LocalUserNum, const FString& ListNam
 			EOS_Friends_GetFriendsCountOptions Options = { };
 			Options.ApiVersion = 1;
 			UE_EOS_CHECK_API_MISMATCH(EOS_FRIENDS_GETFRIENDSCOUNT_API_LATEST, 1);
-			Options.LocalUserId = UserNumToAccountIdMap[LocalUserNum];
+			Options.LocalUserId = GetLocalEpicAccountId(LocalUserNum);
 			int32 FriendCount = EOS_Friends_GetFriendsCount(EOSSubsystem->FriendsHandle, &Options);
 
 			LocalUserNumToFriendsListMap[LocalUserNum]->Empty(FriendCount);
@@ -2049,7 +2083,7 @@ typedef TEOSCallback<EOS_Friends_OnSendInviteCallback, EOS_Friends_SendInviteCal
 
 bool FUserManagerEOS::SendInvite(int32 LocalUserNum, const FUniqueNetId& FriendId, const FString& ListName, const FOnSendInviteComplete& Delegate)
 {
-	if (!UserNumToNetIdMap.Contains(LocalUserNum))
+	if (!LocalUserUniqueNetIds.IsValidIndex(LocalUserNum))
 	{
 		UE_LOG_ONLINE_FRIEND(Warning, TEXT("Can't SendInvite() for user (%d) since they are not logged in"), LocalUserNum);
 		Delegate.ExecuteIfBound(LocalUserNum, false, FriendId, ListName, FString(TEXT("Can't SendInvite() for user (%d) since they are not logged in"), LocalUserNum));
@@ -2082,7 +2116,7 @@ bool FUserManagerEOS::SendInvite(int32 LocalUserNum, const FUniqueNetId& FriendI
 	EOS_Friends_SendInviteOptions Options = { };
 	Options.ApiVersion = 1;
 	UE_EOS_CHECK_API_MISMATCH(EOS_FRIENDS_SENDINVITE_API_LATEST, 1);
-	Options.LocalUserId = UserNumToAccountIdMap[LocalUserNum];
+	Options.LocalUserId = GetLocalEpicAccountId(LocalUserNum);
 	Options.TargetUserId = AccountId;
 	EOS_Friends_SendInvite(EOSSubsystem->FriendsHandle, &Options, CallbackObj, CallbackObj->GetCallbackPtr());
 
@@ -2093,7 +2127,7 @@ typedef TEOSCallback<EOS_Friends_OnAcceptInviteCallback, EOS_Friends_AcceptInvit
 
 bool FUserManagerEOS::AcceptInvite(int32 LocalUserNum, const FUniqueNetId& FriendId, const FString& ListName, const FOnAcceptInviteComplete& Delegate)
 {
-	if (!UserNumToNetIdMap.Contains(LocalUserNum))
+	if (!LocalUserUniqueNetIds.IsValidIndex(LocalUserNum))
 	{
 		UE_LOG_ONLINE_FRIEND(Warning, TEXT("Can't AcceptInvite() for user (%d) since they are not logged in"), LocalUserNum);
 		Delegate.ExecuteIfBound(LocalUserNum, false, FriendId, ListName, FString(TEXT("Can't AcceptInvite() for user (%d) since they are not logged in"), LocalUserNum));
@@ -2126,7 +2160,7 @@ bool FUserManagerEOS::AcceptInvite(int32 LocalUserNum, const FUniqueNetId& Frien
 	EOS_Friends_AcceptInviteOptions Options = { };
 	Options.ApiVersion = 1;
 	UE_EOS_CHECK_API_MISMATCH(EOS_FRIENDS_ACCEPTINVITE_API_LATEST, 1);
-	Options.LocalUserId = UserNumToAccountIdMap[LocalUserNum];
+	Options.LocalUserId = GetLocalEpicAccountId(LocalUserNum);
 	Options.TargetUserId = AccountId;
 	EOS_Friends_AcceptInvite(EOSSubsystem->FriendsHandle, &Options, CallbackObj, CallbackObj->GetCallbackPtr());
 	return true;
@@ -2139,7 +2173,7 @@ void EOS_CALL EOSRejectInviteCallback(const EOS_Friends_RejectInviteCallbackInfo
 
 bool FUserManagerEOS::RejectInvite(int32 LocalUserNum, const FUniqueNetId& FriendId, const FString& ListName)
 {
-	if (!UserNumToNetIdMap.Contains(LocalUserNum))
+	if (!LocalUserUniqueNetIds.IsValidIndex(LocalUserNum))
 	{
 		UE_LOG_ONLINE_FRIEND(Warning, TEXT("Can't RejectInvite() for user (%d) since they are not logged in"), LocalUserNum);
 		return false;
@@ -2156,7 +2190,7 @@ bool FUserManagerEOS::RejectInvite(int32 LocalUserNum, const FUniqueNetId& Frien
 	EOS_Friends_RejectInviteOptions Options{ 0 };
 	Options.ApiVersion = 1;
 	UE_EOS_CHECK_API_MISMATCH(EOS_FRIENDS_REJECTINVITE_API_LATEST, 1);
-	Options.LocalUserId = UserNumToAccountIdMap[LocalUserNum];
+	Options.LocalUserId = GetLocalEpicAccountId(LocalUserNum);
 	Options.TargetUserId = AccountId;
 	EOS_Friends_RejectInvite(EOSSubsystem->FriendsHandle, &Options, nullptr, &EOSRejectInviteCallback);
 	return true;
@@ -2576,7 +2610,7 @@ void FUserManagerEOS::QueryPresence(const FUniqueNetId& UserId, const FOnPresenc
 	EOS_Presence_HasPresenceOptions HasOptions = { };
 	HasOptions.ApiVersion = 1;
 	UE_EOS_CHECK_API_MISMATCH(EOS_PRESENCE_HASPRESENCE_API_LATEST, 1);
-	HasOptions.LocalUserId = UserNumToAccountIdMap[DefaultLocalUser];
+	HasOptions.LocalUserId = GetLocalEpicAccountId();
 	HasOptions.TargetUserId = AccountId;
 	EOS_Bool bHasPresence = EOS_Presence_HasPresence(EOSSubsystem->PresenceHandle, &HasOptions);
 	if (bHasPresence == EOS_FALSE)
@@ -2619,7 +2653,7 @@ void FUserManagerEOS::UpdatePresence(EOS_EpicAccountId AccountId)
 	EOS_Presence_CopyPresenceOptions Options = { };
 	Options.ApiVersion = 3;
 	UE_EOS_CHECK_API_MISMATCH(EOS_PRESENCE_COPYPRESENCE_API_LATEST, 3);
-	Options.LocalUserId = UserNumToAccountIdMap[DefaultLocalUser];
+	Options.LocalUserId = GetLocalEpicAccountId();
 	Options.TargetUserId = AccountId;
 	EOS_EResult CopyResult = EOS_Presence_CopyPresence(EOSSubsystem->PresenceHandle, &Options, &PresenceInfo);
 	if (CopyResult == EOS_EResult::EOS_Success)
@@ -2703,7 +2737,7 @@ bool FUserManagerEOS::QueryUserInfo(int32 LocalUserNum, const TArray<FUniqueNetI
 	{
 		const FUniqueNetIdEOS& EOSID = FUniqueNetIdEOS::Cast(*NetId);
 		// Skip querying for local users since we already have that data
-		if (StringToUserAccountMap.Contains(EOSID.ToString()))
+		if (IsLocalUser(EOSID))
 		{
 			continue;
 		}
@@ -2773,7 +2807,7 @@ void FUserManagerEOS::ReadUserInfo(int32 LocalUserNum, EOS_EpicAccountId EpicAcc
 	EOS_UserInfo_QueryUserInfoOptions Options = { };
 	Options.ApiVersion = 1;
 	UE_EOS_CHECK_API_MISMATCH(EOS_USERINFO_QUERYUSERINFO_API_LATEST, 1);
-	Options.LocalUserId = UserNumToAccountIdMap[DefaultLocalUser];
+	Options.LocalUserId = GetLocalEpicAccountId(LocalUserNum);
 	Options.TargetUserId = EpicAccountId;
 	EOS_UserInfo_QueryUserInfo(EOSSubsystem->UserInfoHandle, &Options, CallbackObj, CallbackObj->GetCallbackPtr());
 
@@ -2793,7 +2827,7 @@ bool FUserManagerEOS::GetAllUserInfo(int32 LocalUserNum, TArray<TSharedRef<FOnli
 		}
 	}
 	// Get local users
-	for (TMap<FString, FUserOnlineAccountEOSRef>::TConstIterator It(StringToUserAccountMap); It; ++It)
+	for (TUniqueNetIdMap<FUserOnlineAccountEOSRef>::TConstIterator It(UniqueNetIdToUserAccountMap); It; ++It)
 	{
 		OutUsers.Add(It.Value());
 	}
@@ -2853,7 +2887,7 @@ bool FUserManagerEOS::QueryUserIdMapping(const FUniqueNetId& UserId, const FStri
 		if (bWasSuccessful)
 		{
 			const FString NetIdStr = LexToString(Data->TargetUserId);
-			FUniqueNetIdEOSPtr LocalUserId = UserNumToNetIdMap[DefaultLocalUser];
+			const FUniqueNetIdEOSRef LocalUserId = LocalUserUniqueNetIds[LocalUserNum];
 			if (!EpicAccountIdToOnlineUserMap.Contains(Data->TargetUserId))
 			{
 				// Registering the player will also query the presence/user info data
@@ -2973,10 +3007,10 @@ bool FUserManagerEOS::QueryExternalIdMappings(const FUniqueNetId& UserId, const
 			FUniqueNetIdEOSPtr EOSID = FUniqueNetIdEOS::EmptyId();
 			if (Result == EOS_EResult::EOS_Success)
 			{
-				EOSID = UserNumToNetIdMap[LocalUserNum];
+				EOSID = LocalUserUniqueNetIds[LocalUserNum];
 
 				FGetAccountMappingOptions Options;
-				Options.LocalUserId = UserNumToProductUserIdMap[DefaultLocalUser];
+				Options.LocalUserId = GetLocalProductUserId(LocalUserNum);
 				// Get the product id for each epic account passed in
 				for (const FString& StringId : BatchIds)
 				{
diff --git a/Engine/Plugins/Online/OnlineSubsystemEOS/Source/OnlineSubsystemEOS/Private/UserManagerEOS.h b/Engine/Plugins/Online/OnlineSubsystemEOS/Source/OnlineSubsystemEOS/Private/UserManagerEOS.h
--- a/Engine/Plugins/Online/OnlineSubsystemEOS/Source/OnlineSubsystemEOS/Private/UserManagerEOS.h
+++ b/Engine/Plugins/Online/OnlineSubsystemEOS/Source/OnlineSubsystemEOS/Private/UserManagerEOS.h
@@ -12,6 +12,11 @@
 	#include "eos_friends_types.h"
 	#include "eos_connect_types.h"
 
+#define INVALID_LOCAL_USER -1
+
+template <typename ValueType>
+using TLocalUserArray = TSparseArray<ValueType, TInlineSparseArrayAllocator<MAX_LOCAL_PLAYERS>>;
+
 class FOnlineSubsystemEOS;
 class IOnlineSubsystem;
 
@@ -279,22 +284,23 @@ public:
 	virtual FUniqueNetIdPtr GetExternalIdMapping(const FExternalIdQueryOptions& QueryOptions, const FString& ExternalId) override;
 // ~IOnlineUser Interface
 
+	const FUniqueNetIdEOSPtr GetLocalUniqueNetIdEOS(int32 LocalUserNum) const;
+	const FUniqueNetIdEOSPtr GetLocalUniqueNetIdEOS() const;
 	EOS_EpicAccountId GetLocalEpicAccountId(int32 LocalUserNum) const;
 	EOS_EpicAccountId GetLocalEpicAccountId() const;
 	EOS_ProductUserId GetLocalProductUserId(int32 LocalUserNum) const;
 	EOS_ProductUserId GetLocalProductUserId() const;
-	EOS_EpicAccountId GetLocalEpicAccountId(EOS_ProductUserId UserId) const;
-	EOS_ProductUserId GetLocalProductUserId(EOS_EpicAccountId AccountId) const;
-	FUniqueNetIdEOSPtr GetLocalUniqueNetIdEOS(int32 LocalUserNum) const;
-	FUniqueNetIdEOSPtr GetLocalUniqueNetIdEOS(EOS_ProductUserId UserId) const;
-	FUniqueNetIdEOSPtr GetLocalUniqueNetIdEOS(EOS_EpicAccountId AccountId) const;
-	FUniqueNetIdEOSPtr GetLocalUniqueNetIdEOS() const
-	{
-		return GetLocalUniqueNetIdEOS(GetDefaultLocalUser());
-	}
 
 	int32 GetLocalUserNumFromUniqueNetId(const FUniqueNetId& NetId) const;
+	int32 GetLocalUserNumFromEpicAccountId(const EOS_EpicAccountId& EpicAccountId) const;
+	int32 GetLocalUserNumFromProductUserId(const EOS_ProductUserId& ProductUserId) const;
+
+	const FUniqueNetIdEOSPtr GetLocalUniqueNetIdEOS(const EOS_EpicAccountId& EpicAccountId) const;
+	const FUniqueNetIdEOSPtr GetLocalUniqueNetIdEOS(const EOS_ProductUserId& ProductUserId) const;
+
 	bool IsLocalUser(const FUniqueNetId& NetId) const;
+	bool IsLocalUser(const EOS_EpicAccountId& EpicAccountId) const;
+	bool IsLocalUser(const EOS_ProductUserId& ProductUserId) const;
 
 	typedef TFunction<void(TMap<EOS_ProductUserId, FUniqueNetIdEOSRef> ResolvedUniqueNetIds)> FResolveUniqueNetIdsCallback;
 	typedef TFunction<void(FUniqueNetIdEOSRef ResolvedUniqueNetId)> FResolveUniqueNetIdCallback;
@@ -362,12 +368,9 @@ private:
 	TMap<int32, FNotificationIdCallbackPair*> LocalUserNumToConnectLoginNotifcationMap;
 
 	/** Ids mapped to locally registered users */
-	TMap<int32, EOS_EpicAccountId> UserNumToAccountIdMap;
-	TMap<EOS_EpicAccountId, int32> AccountIdToUserNumMap;
-	TMap<int32, FUniqueNetIdEOSPtr> UserNumToNetIdMap;
-	TMap<int32, EOS_ProductUserId> UserNumToProductUserIdMap;
-	TMap<EOS_ProductUserId, int32> ProductUserIdToUserNumMap;
-	TMap<FString, FUserOnlineAccountEOSRef> StringToUserAccountMap;
+	TLocalUserArray<FUniqueNetIdEOSRef> LocalUserUniqueNetIds;
+	/** General map for all registered online users */
+	TUniqueNetIdMap<FUserOnlineAccountEOSRef> UniqueNetIdToUserAccountMap;
 
 	/** General account mappings */
 	TMap<EOS_EpicAccountId, FString> AccountIdToStringMap;
