- Fix falling back on any entry when -epicsandbox is passed, we should require a matching entry

diff --git a/Engine/Plugins/Online/OnlineSubsystemEOS/Source/OnlineSubsystemEOS/Private/EOSSettings.cpp b/Engine/Plugins/Online/OnlineSubsystemEOS/Source/OnlineSubsystemEOS/Private/EOSSettings.cpp
--- a/Engine/Plugins/Online/OnlineSubsystemEOS/Source/OnlineSubsystemEOS/Private/EOSSettings.cpp
+++ b/Engine/Plugins/Online/OnlineSubsystemEOS/Source/OnlineSubsystemEOS/Private/EOSSettings.cpp
@@ -208,22 +208,22 @@ bool UEOSSettings::GetSelectedArtifactSettings(FEOSArtifactSettings& OutSettings
 	// Prefer -EOSArtifactNameOverride over previous.
 	FParse::Value(FCommandLine::Get(), TEXT("EOSArtifactNameOverride="), ArtifactName);
 
-	// Get the sandbox id, which only comes in via EGS. If present, grab the settings where both match.
+	const TCHAR* CheckConfigText = TEXT("please check there is a matching config entry in the \"Artifacts\" array in the \"" INI_SECTION "\" section of the Engine.ini hierarchy.");
+
+	// Get the sandbox id, normally passed by EGS. If present, grab the settings where both artifact and sandbox match, otherwise just match on artifact.
 	FString SandboxId;
 	if (FParse::Value(FCommandLine::Get(), TEXT("EpicSandboxId="), SandboxId))
 	{
-		if (GetArtifactSettings(ArtifactName, SandboxId, OutSettings))
-		{
-			return true;
-		}
-
-		UE_LOG_ONLINE(Warning, TEXT("UEOSSettings::GetSelectedArtifactSettings() ArtifactName=[%s] SandboxId=[%s] no settings found."), *ArtifactName, *SandboxId);
+		const bool bSuccess = GetArtifactSettings(ArtifactName, SandboxId, OutSettings);
+		UE_CLOG_ONLINE(!bSuccess, Error, TEXT("%hs ArtifactName=[%s] SandboxId=[%s] no settings found for this Artifact/Sandbox combination, %s"), __FUNCTION__, *ArtifactName, *SandboxId, CheckConfigText);
+		return bSuccess;
+	}
+	else
+	{
+		const bool bSuccess = GetArtifactSettings(ArtifactName, OutSettings);
+		UE_CLOG_ONLINE(!bSuccess, Error, TEXT("%hs ArtifactName=[%s] no settings found for this Artifact, %s"), __FUNCTION__, *ArtifactName, CheckConfigText);
+		return bSuccess;
 	}
-
-	// Fall back on just matching the Artifact name. This assumes non-EGS and only one settings entry per ArtifactName in config.
-	const bool bSuccess = GetArtifactSettings(ArtifactName, OutSettings);
-	UE_CLOG_ONLINE(!bSuccess, Error, TEXT("UEOSSettings::GetSelectedArtifactSettings() ArtifactName=[%s] no settings found."), *ArtifactName);
-	return bSuccess;
 }
 
 FString UEOSSettings::GetDefaultArtifactName()
