Removal of UserManagerEOS methods with a default local user, refactor in all callsites to use the correct (or nearest) local user.

diff --git a/Engine/Plugins/Online/OnlineSubsystemEOS/Source/OnlineSubsystemEOS/Private/OnlineSessionEOS.cpp b/Engine/Plugins/Online/OnlineSubsystemEOS/Source/OnlineSubsystemEOS/Private/OnlineSessionEOS.cpp
--- a/Engine/Plugins/Online/OnlineSubsystemEOS/Source/OnlineSubsystemEOS/Private/OnlineSessionEOS.cpp
+++ b/Engine/Plugins/Online/OnlineSubsystemEOS/Source/OnlineSubsystemEOS/Private/OnlineSessionEOS.cpp
@@ -621,6 +621,24 @@ FOnlineSession* FOnlineSessionEOS::GetOnlineSessionFromLobbyId(const FUniqueNetI
 	return Result;
 }
 
+int32 FOnlineSessionEOS::GetDefaultLocalUserForLobby(const FUniqueNetIdString& SessionId)
+{
+	if (FOnlineSession * Session = GetOnlineSessionFromLobbyId(FUniqueNetIdEOSLobby::Cast(SessionId)))
+	{
+		for (const TPair<FUniqueNetIdRef, FSessionSettings>& MemberSettings : Session->SessionSettings.MemberSettings)
+		{
+			int32 LocalUserId = EOSSubsystem->UserManager->GetLocalUserNumFromUniqueNetId(*MemberSettings.Key);
+			
+			if (LocalUserId != INVALID_LOCAL_USER)
+			{
+				return LocalUserId;
+			}
+		}
+	}
+
+	return INVALID_LOCAL_USER;
+}
+
 void FOnlineSessionEOS::RegisterLobbyNotifications()
 {
 	// Lobby data updates
@@ -632,7 +650,7 @@ void FOnlineSessionEOS::RegisterLobbyNotifications()
 	LobbyUpdateReceivedCallback = LobbyUpdateReceivedCallbackObj;
 	LobbyUpdateReceivedCallbackObj->CallbackLambda = [this](const EOS_Lobby_LobbyUpdateReceivedCallbackInfo* Data)
 	{
-		OnLobbyUpdateReceived(Data->LobbyId, EOSSubsystem->UserManager->GetLocalProductUserId());
+		OnLobbyUpdateReceived(Data->LobbyId);
 	};
 
 	LobbyUpdateReceivedId = EOS_Lobby_AddNotifyLobbyUpdateReceived(LobbyHandle, &AddNotifyLobbyUpdateReceivedOptions, LobbyUpdateReceivedCallbackObj, LobbyUpdateReceivedCallbackObj->GetCallbackPtr());
@@ -696,9 +714,11 @@ void FOnlineSessionEOS::RegisterLobbyNotifications()
 	JoinLobbyAcceptedId = EOS_Lobby_AddNotifyJoinLobbyAccepted(LobbyHandle, &AddNotifyJoinLobbyAcceptedOptions, JoinLobbyAcceptedCallbackObj, JoinLobbyAcceptedCallbackObj->GetCallbackPtr());
 }
 
-void FOnlineSessionEOS::OnLobbyUpdateReceived(const EOS_LobbyId& LobbyId, const EOS_ProductUserId& LocalUserId)
+void FOnlineSessionEOS::OnLobbyUpdateReceived(const EOS_LobbyId& LobbyId)
 {
 	const FUniqueNetIdEOSLobbyRef LobbyNetId = FUniqueNetIdEOSLobby::Create(UTF8_TO_TCHAR(LobbyId));
+	const EOS_ProductUserId LocalUserId = EOSSubsystem->UserManager->GetLocalProductUserId(GetDefaultLocalUserForLobby(*LobbyNetId));
+
 	FNamedOnlineSession* Session = GetNamedSessionFromLobbyId(*LobbyNetId);
 	if (Session)
 	{
@@ -755,7 +775,7 @@ void FOnlineSessionEOS::OnLobbyMemberUpdateReceived(const EOS_LobbyId& LobbyId,
 {
 	const FUniqueNetIdEOSLobbyRef LobbyNetId = FUniqueNetIdEOSLobby::Create(UTF8_TO_TCHAR(LobbyId));
 
-	EOSSubsystem->UserManager->ResolveUniqueNetId(TargetUserId, [this, TargetUserId, LobbyNetId](FUniqueNetIdEOSRef ResolvedUniqueNetId)
+	EOSSubsystem->UserManager->ResolveUniqueNetId(EOSSubsystem->UserManager->GetDefaultLocalUser(), TargetUserId, [this, TargetUserId, LobbyNetId](FUniqueNetIdEOSRef ResolvedUniqueNetId)
 		{
 			UpdateOrAddLobbyMember(LobbyNetId, ResolvedUniqueNetId);
 		});
@@ -771,7 +791,7 @@ void FOnlineSessionEOS::OnMemberStatusReceived(const EOS_LobbyId& LobbyId, const
 		{
 		case EOS_ELobbyMemberStatus::EOS_LMS_JOINED:
 			{
-				EOSSubsystem->UserManager->ResolveUniqueNetId(TargetUserId, [this, LobbyNetId](FUniqueNetIdEOSRef ResolvedUniqueNetId)
+				EOSSubsystem->UserManager->ResolveUniqueNetId(EOSSubsystem->UserManager->GetDefaultLocalUser(), TargetUserId, [this, LobbyNetId](FUniqueNetIdEOSRef ResolvedUniqueNetId)
 					{
 						UpdateOrAddLobbyMember(LobbyNetId, ResolvedUniqueNetId);
 					});
@@ -780,7 +800,7 @@ void FOnlineSessionEOS::OnMemberStatusReceived(const EOS_LobbyId& LobbyId, const
 			break;
 		case EOS_ELobbyMemberStatus::EOS_LMS_LEFT:
 			{
-				EOSSubsystem->UserManager->ResolveUniqueNetId(TargetUserId, [this, LobbyNetId](FUniqueNetIdEOSRef ResolvedUniqueNetId)
+				EOSSubsystem->UserManager->ResolveUniqueNetId(EOSSubsystem->UserManager->GetDefaultLocalUser(), TargetUserId, [this, LobbyNetId](FUniqueNetIdEOSRef ResolvedUniqueNetId)
 					{
 						FNamedOnlineSession* Session = GetNamedSessionFromLobbyId(*LobbyNetId);
 						if (Session)
@@ -798,7 +818,7 @@ void FOnlineSessionEOS::OnMemberStatusReceived(const EOS_LobbyId& LobbyId, const
 			break;
 		case EOS_ELobbyMemberStatus::EOS_LMS_DISCONNECTED:
 			{
-				EOSSubsystem->UserManager->ResolveUniqueNetId(TargetUserId, [this, LobbyNetId](FUniqueNetIdEOSRef ResolvedUniqueNetId)
+				EOSSubsystem->UserManager->ResolveUniqueNetId(EOSSubsystem->UserManager->GetDefaultLocalUser(), TargetUserId, [this, LobbyNetId](FUniqueNetIdEOSRef ResolvedUniqueNetId)
 					{
 						FNamedOnlineSession* Session = GetNamedSessionFromLobbyId(*LobbyNetId);
 						if (Session)
@@ -816,7 +836,7 @@ void FOnlineSessionEOS::OnMemberStatusReceived(const EOS_LobbyId& LobbyId, const
 			break;
 		case EOS_ELobbyMemberStatus::EOS_LMS_KICKED:
 			{
-				EOSSubsystem->UserManager->ResolveUniqueNetId(TargetUserId, [this, LobbyNetId](FUniqueNetIdEOSRef ResolvedUniqueNetId)
+				EOSSubsystem->UserManager->ResolveUniqueNetId(EOSSubsystem->UserManager->GetDefaultLocalUser(), TargetUserId, [this, LobbyNetId](FUniqueNetIdEOSRef ResolvedUniqueNetId)
 					{
 						FNamedOnlineSession* Session = GetNamedSessionFromLobbyId(*LobbyNetId);
 						if (Session)
@@ -834,7 +854,7 @@ void FOnlineSessionEOS::OnMemberStatusReceived(const EOS_LobbyId& LobbyId, const
 			break;
 		case EOS_ELobbyMemberStatus::EOS_LMS_PROMOTED:
 			{
-				EOSSubsystem->UserManager->ResolveUniqueNetId(TargetUserId, [this, LobbyNetId](FUniqueNetIdEOSRef ResolvedUniqueNetId)
+				EOSSubsystem->UserManager->ResolveUniqueNetId(EOSSubsystem->UserManager->GetDefaultLocalUser(), TargetUserId, [this, LobbyNetId](FUniqueNetIdEOSRef ResolvedUniqueNetId)
 					{
 						FNamedOnlineSession* Session = GetNamedSessionFromLobbyId(*LobbyNetId);
 						if (Session)
@@ -866,7 +886,7 @@ void FOnlineSessionEOS::OnMemberStatusReceived(const EOS_LobbyId& LobbyId, const
 			break;
 		case EOS_ELobbyMemberStatus::EOS_LMS_CLOSED:
 			{
-				EOSSubsystem->UserManager->ResolveUniqueNetId(TargetUserId, [this, LobbyNetId](FUniqueNetIdEOSRef ResolvedUniqueNetId)
+				EOSSubsystem->UserManager->ResolveUniqueNetId(EOSSubsystem->UserManager->GetDefaultLocalUser(), TargetUserId, [this, LobbyNetId](FUniqueNetIdEOSRef ResolvedUniqueNetId)
 					{
 						FNamedOnlineSession* Session = GetNamedSessionFromLobbyId(*LobbyNetId);
 						if (Session)
@@ -1789,7 +1809,7 @@ bool FOnlineSessionEOS::DestroySession(FName SessionName, const FOnDestroySessio
 
 				if (Session->SessionSettings.bUseLobbiesIfAvailable)
 				{
-					Result = DestroyLobbySession(Session, CompletionDelegate);
+					Result = DestroyLobbySession(EOSSubsystem->UserManager->GetDefaultLocalUser(), Session, CompletionDelegate);
 				}
 				else
 				{
@@ -2983,7 +3003,7 @@ void FOnlineSessionEOS::UpdateOrAddLobbyMember(const FUniqueNetIdEOSLobbyRef& Lo
 			Options.ApiVersion = 1;
 			UE_EOS_CHECK_API_MISMATCH(EOS_LOBBY_COPYLOBBYDETAILSHANDLE_API_LATEST, 1);
 			Options.LobbyId = (EOS_LobbyId)Utf8LobbyId.Get();
-			Options.LocalUserId = EOSSubsystem->UserManager->GetLocalProductUserId();
+			Options.LocalUserId = EOSSubsystem->UserManager->GetLocalProductUserId(GetDefaultLocalUserForLobby(*LobbyNetId));
 
 			EOS_HLobbyDetails LobbyDetailsHandle;
 
@@ -3768,7 +3788,7 @@ uint32 FOnlineSessionEOS::CreateLobbySession(int32 HostingPlayerNum, FNamedOnlin
 	FName SessionName = Session->SessionName;
 	FLobbyCreatedCallback* CallbackObj = new FLobbyCreatedCallback(FOnlineSessionEOSWeakPtr(AsShared()));
 	LobbyCreatedCallback = CallbackObj;
-	CallbackObj->CallbackLambda = [this, SessionName, LocalProductUserId, LocalUserNetId](const EOS_Lobby_CreateLobbyCallbackInfo* Data)
+	CallbackObj->CallbackLambda = [this, HostingPlayerNum, SessionName, LocalProductUserId, LocalUserNetId](const EOS_Lobby_CreateLobbyCallbackInfo* Data)
 	{
 		FNamedOnlineSession* Session = GetNamedSession(SessionName);
 		if (Session)
@@ -3796,7 +3816,7 @@ uint32 FOnlineSessionEOS::CreateLobbySession(int32 HostingPlayerNum, FNamedOnlin
 
 				BeginSessionAnalytics(Session);
 
-				UpdateLobbySession(Session, FOnUpdateSessionCompleteDelegate::CreateThreadSafeSP(this, &FOnlineSessionEOS::OnCreateLobbySessionUpdateComplete));
+				UpdateLobbySession(Session, FOnUpdateSessionCompleteDelegate::CreateThreadSafeSP(this, &FOnlineSessionEOS::OnCreateLobbySessionUpdateComplete, HostingPlayerNum));
 			}
 			else
 			{
@@ -3816,7 +3836,7 @@ uint32 FOnlineSessionEOS::CreateLobbySession(int32 HostingPlayerNum, FNamedOnlin
 	return ONLINE_IO_PENDING;
 }
 
-void FOnlineSessionEOS::OnCreateLobbySessionUpdateComplete(FName SessionName, bool bWasSuccessful)
+void FOnlineSessionEOS::OnCreateLobbySessionUpdateComplete(FName SessionName, bool bWasSuccessful, int32 HostingPlayerNum)
 {
 	if (!bWasSuccessful)
 	{
@@ -3825,7 +3845,7 @@ void FOnlineSessionEOS::OnCreateLobbySessionUpdateComplete(FName SessionName, bo
 		FNamedOnlineSession* Session = GetNamedSession(SessionName);
 		if (Session)
 		{
-			DestroyLobbySessionOnCreationUpdateError(Session);
+			DestroyLobbySessionOnCreationUpdateError(HostingPlayerNum, Session);
 		}
 	}
 	else
@@ -3834,7 +3854,7 @@ void FOnlineSessionEOS::OnCreateLobbySessionUpdateComplete(FName SessionName, bo
 	}
 }
 
-void FOnlineSessionEOS::DestroyLobbySessionOnCreationUpdateError(FNamedOnlineSession* Session)
+void FOnlineSessionEOS::DestroyLobbySessionOnCreationUpdateError(int32 LocalUserNum, FNamedOnlineSession* Session)
 {
 	check(Session != nullptr);
 	check(Session->SessionInfo.IsValid());
@@ -3847,12 +3867,12 @@ void FOnlineSessionEOS::DestroyLobbySessionOnCreationUpdateError(FNamedOnlineSes
 	UE_EOS_CHECK_API_MISMATCH(EOS_LOBBY_DESTROYLOBBY_API_LATEST, 1);
 	const FTCHARToUTF8 Utf8LobbyId(*SessionInfo->GetSessionId().ToString());
 	DestroyOptions.LobbyId = (EOS_LobbyId)Utf8LobbyId.Get();
-	DestroyOptions.LocalUserId = EOSSubsystem->UserManager->GetLocalProductUserId(EOSSubsystem->UserManager->GetDefaultLocalUser()); // Maybe not split screen friendly
+	DestroyOptions.LocalUserId = EOSSubsystem->UserManager->GetLocalProductUserId(LocalUserNum);
 
 	FName SessionName = Session->SessionName;
 	FLobbyDestroyedCallback* DestroyCallbackObj = new FLobbyDestroyedCallback(FOnlineSessionEOSWeakPtr(AsShared()));
 	LobbyLeftCallback = DestroyCallbackObj;
-	DestroyCallbackObj->CallbackLambda = [this, SessionName](const EOS_Lobby_DestroyLobbyCallbackInfo* Data)
+	DestroyCallbackObj->CallbackLambda = [this, LocalUserNum, SessionName](const EOS_Lobby_DestroyLobbyCallbackInfo* Data)
 	{
 		FNamedOnlineSession* LobbySession = GetNamedSession(SessionName);
 		if (LobbySession)
@@ -3868,7 +3888,7 @@ void FOnlineSessionEOS::DestroyLobbySessionOnCreationUpdateError(FNamedOnlineSes
 			}
 
 #if WITH_EOS_RTC
-			if (FEOSVoiceChatUser* VoiceChatUser = static_cast<FEOSVoiceChatUser*>(EOSSubsystem->GetEOSVoiceChatUserInterface(*EOSSubsystem->UserManager->GetLocalUniqueNetIdEOS())))
+			if (FEOSVoiceChatUser* VoiceChatUser = static_cast<FEOSVoiceChatUser*>(EOSSubsystem->GetEOSVoiceChatUserInterface(*EOSSubsystem->UserManager->GetLocalUniqueNetIdEOS(LocalUserNum))))
 			{
 				VoiceChatUser->RemoveLobbyRoom(UTF8_TO_TCHAR(Data->LobbyId));
 			}
@@ -3946,7 +3966,7 @@ uint32 FOnlineSessionEOS::JoinLobbySession(int32 PlayerNum, FNamedOnlineSession*
 						}
 #endif
 
-						OnLobbyUpdateReceived(Data->LobbyId, LocalUserNetId->GetProductUserId());
+						OnLobbyUpdateReceived(Data->LobbyId); // We could use LocalUserNetId here instead of the default local user for the session, but the end result should be the same
 					}
 					else
 					{
@@ -4238,7 +4258,7 @@ uint32 FOnlineSessionEOS::EndLobbySession(FNamedOnlineSession* Session)
 	return ONLINE_IO_PENDING;
 }
 
-uint32 FOnlineSessionEOS::DestroyLobbySession(FNamedOnlineSession* Session, const FOnDestroySessionCompleteDelegate& CompletionDelegate)
+uint32 FOnlineSessionEOS::DestroyLobbySession(int32 LocalUserNum, FNamedOnlineSession* Session, const FOnDestroySessionCompleteDelegate& CompletionDelegate)
 {
 	check(Session != nullptr);
 
@@ -4257,12 +4277,12 @@ uint32 FOnlineSessionEOS::DestroyLobbySession(FNamedOnlineSession* Session, cons
 		UE_EOS_CHECK_API_MISMATCH(EOS_LOBBY_LEAVELOBBY_API_LATEST, 1);
 		const FTCHARToUTF8 Utf8LobbyId(*SessionInfo->GetSessionId().ToString());
 		LeaveOptions.LobbyId = (EOS_LobbyId)Utf8LobbyId.Get();
-		LeaveOptions.LocalUserId = EOSSubsystem->UserManager->GetLocalProductUserId(EOSSubsystem->UserManager->GetDefaultLocalUser()); // Maybe not split screen friendly
+		LeaveOptions.LocalUserId = EOSSubsystem->UserManager->GetLocalProductUserId(LocalUserNum);
 
 		FName SessionName = Session->SessionName;
 		FLobbyLeftCallback* LeaveCallbackObj = new FLobbyLeftCallback(FOnlineSessionEOSWeakPtr(AsShared()));
 		LobbyLeftCallback = LeaveCallbackObj;
-		LeaveCallbackObj->CallbackLambda = [this, SessionName, CompletionDelegate](const EOS_Lobby_LeaveLobbyCallbackInfo* Data)
+		LeaveCallbackObj->CallbackLambda = [this, LocalUserNum, SessionName, CompletionDelegate](const EOS_Lobby_LeaveLobbyCallbackInfo* Data)
 		{
 			FNamedOnlineSession* LobbySession = GetNamedSession(SessionName);
 			if (LobbySession)
@@ -4278,7 +4298,7 @@ uint32 FOnlineSessionEOS::DestroyLobbySession(FNamedOnlineSession* Session, cons
 				}
 
 #if WITH_EOS_RTC
-				if (FEOSVoiceChatUser* VoiceChatUser = static_cast<FEOSVoiceChatUser*>(EOSSubsystem->GetEOSVoiceChatUserInterface(*EOSSubsystem->UserManager->GetLocalUniqueNetIdEOS())))
+				if (FEOSVoiceChatUser* VoiceChatUser = static_cast<FEOSVoiceChatUser*>(EOSSubsystem->GetEOSVoiceChatUserInterface(*EOSSubsystem->UserManager->GetLocalUniqueNetIdEOS(LocalUserNum))))
 				{
 					VoiceChatUser->RemoveLobbyRoom(UTF8_TO_TCHAR(Data->LobbyId));
 				}
@@ -4556,7 +4576,7 @@ void FOnlineSessionEOS::CopyLobbyData(const TSharedRef<FLobbyDetailsEOS>& LobbyD
 
 	if (!TargetUserIds.IsEmpty())
 	{
-		EOSSubsystem->UserManager->ResolveUniqueNetIds(TargetUserIds, [this, LobbyDetails, LobbyId = FUniqueNetIdEOSLobby::Create(UTF8_TO_TCHAR(LobbyDetailsInfo->LobbyId)), OriginalCallback = Callback](TMap<EOS_ProductUserId, FUniqueNetIdEOSRef> ResolvedUniqueNetIds)
+		EOSSubsystem->UserManager->ResolveUniqueNetIds(EOSSubsystem->UserManager->GetDefaultLocalUser(), TargetUserIds, [this, LobbyDetails, LobbyId = FUniqueNetIdEOSLobby::Create(UTF8_TO_TCHAR(LobbyDetailsInfo->LobbyId)), OriginalCallback = Callback](TMap<EOS_ProductUserId, FUniqueNetIdEOSRef> ResolvedUniqueNetIds)
 			{
 				
 				FOnlineSession* Session = GetOnlineSessionFromLobbyId(*LobbyId);
diff --git a/Engine/Plugins/Online/OnlineSubsystemEOS/Source/OnlineSubsystemEOS/Private/OnlineSessionEOS.h b/Engine/Plugins/Online/OnlineSubsystemEOS/Source/OnlineSubsystemEOS/Private/OnlineSessionEOS.h
--- a/Engine/Plugins/Online/OnlineSubsystemEOS/Source/OnlineSubsystemEOS/Private/OnlineSessionEOS.h
+++ b/Engine/Plugins/Online/OnlineSubsystemEOS/Source/OnlineSubsystemEOS/Private/OnlineSessionEOS.h
@@ -224,15 +224,15 @@ private:
 	FCallbackBase* LobbySendInviteCallback;
 
 	uint32 CreateLobbySession(int32 HostingPlayerNum, FNamedOnlineSession* Session);
-	void OnCreateLobbySessionUpdateComplete(FName SessionName, bool bWasSuccessful);
-	void DestroyLobbySessionOnCreationUpdateError(FNamedOnlineSession* Session);
+	void OnCreateLobbySessionUpdateComplete(FName SessionName, bool bWasSuccessful, int32 HostingPlayerNum);
+	void DestroyLobbySessionOnCreationUpdateError(int32 LocalUserNum, FNamedOnlineSession* Session);
 	uint32 FindLobbySession(int32 SearchingPlayerNum, const TSharedRef<FOnlineSessionSearch>& SearchSettings);
 	void StartLobbySearch(int32 SearchingPlayerNum, EOS_HLobbySearch LobbySearchHandle, const TSharedRef<FOnlineSessionSearch>& SearchSettings, const FOnSingleSessionResultCompleteDelegate& CompletionDelegate);
 	uint32 JoinLobbySession(int32 PlayerNum, FNamedOnlineSession* Session, const FOnlineSession* SearchSession);
 	uint32 UpdateLobbySession(FNamedOnlineSession* Session, const FOnUpdateSessionCompleteDelegate& CompletionDelegate);
 	uint32 StartLobbySession(FNamedOnlineSession* Session);
 	uint32 EndLobbySession(FNamedOnlineSession* Session);
-	uint32 DestroyLobbySession(FNamedOnlineSession* Session, const FOnDestroySessionCompleteDelegate& CompletionDelegate);
+	uint32 DestroyLobbySession(int32 LocalUserNum, FNamedOnlineSession* Session, const FOnDestroySessionCompleteDelegate& CompletionDelegate);
 	bool SendLobbyInvite(FName SessionName, EOS_ProductUserId SenderId, EOS_ProductUserId ReceiverId);
 
 	// Lobby notification callbacks and methods
@@ -247,7 +247,7 @@ private:
 	EOS_NotificationId JoinLobbyAcceptedId;
 	FCallbackBase* JoinLobbyAcceptedCallback;
 
-	void OnLobbyUpdateReceived(const EOS_LobbyId& LobbyId, const EOS_ProductUserId& LocalUserId);
+	void OnLobbyUpdateReceived(const EOS_LobbyId& LobbyId);
 	void OnLobbyMemberUpdateReceived(const EOS_LobbyId& LobbyId, const EOS_ProductUserId& TargetUserId);
 	void OnMemberStatusReceived(const EOS_LobbyId& LobbyId, const EOS_ProductUserId& TargetUserId, EOS_ELobbyMemberStatus CurrentStatus);
 	void OnLobbyInviteAccepted(const char* InviteId, const EOS_ProductUserId& LocalUserId, const EOS_ProductUserId& TargetUserId);
@@ -328,6 +328,8 @@ private:
 	static void SetPortFromNetDriver(const FOnlineSubsystemEOS& Subsystem, const TSharedPtr<FOnlineSessionInfo>& SessionInfo);
 	bool IsHost(const FNamedOnlineSession& Session) const;
 
+	int32 GetDefaultLocalUserForLobby(const FUniqueNetIdString& SessionId);
+
 	/** Reference to the main EOS subsystem */
 	FOnlineSubsystemEOS* EOSSubsystem;
 
diff --git a/Engine/Plugins/Online/OnlineSubsystemEOS/Source/OnlineSubsystemEOS/Private/OnlineStoreEOS.cpp b/Engine/Plugins/Online/OnlineSubsystemEOS/Source/OnlineSubsystemEOS/Private/OnlineStoreEOS.cpp
--- a/Engine/Plugins/Online/OnlineSubsystemEOS/Source/OnlineSubsystemEOS/Private/OnlineStoreEOS.cpp
+++ b/Engine/Plugins/Online/OnlineSubsystemEOS/Source/OnlineSubsystemEOS/Private/OnlineStoreEOS.cpp
@@ -393,7 +393,7 @@ bool FOnlineStoreEOS::HandleEcomExec(UWorld* InWorld, const TCHAR* Cmd, FOutputD
 {
 	if (FParse::Command(&Cmd, TEXT("OFFERS")))
 	{
-		QueryOffers(*EOSSubsystem->UserManager->GetLocalUniqueNetIdEOS(),
+		QueryOffers(*EOSSubsystem->UserManager->GetLocalUniqueNetIdEOS(EOSSubsystem->UserManager->GetDefaultLocalUser()),
 			FOnQueryOnlineStoreOffersComplete::CreateLambda([this](bool bWasSuccessful, const TArray<FUniqueOfferId>& OfferIds, const FString& ErrorStr)
 		{
 			UE_LOG_ONLINE(Error, TEXT("QueryOffers: %s with error (%s)"), bWasSuccessful ? TEXT("succeeded") : TEXT("failed"), *ErrorStr);
@@ -407,7 +407,7 @@ bool FOnlineStoreEOS::HandleEcomExec(UWorld* InWorld, const TCHAR* Cmd, FOutputD
 	}
 	else if (FParse::Command(&Cmd, TEXT("RECEIPTS")))
 	{
-		QueryReceipts(*EOSSubsystem->UserManager->GetLocalUniqueNetIdEOS(), false,
+		QueryReceipts(*EOSSubsystem->UserManager->GetLocalUniqueNetIdEOS(EOSSubsystem->UserManager->GetDefaultLocalUser()), false,
 			FOnQueryReceiptsComplete::CreateLambda([this](const FOnlineError& Result)
 		{
 			UE_LOG_ONLINE(Log, TEXT("QueryReceipts: %s with error (%s)"), Result.WasSuccessful() ? TEXT("succeeded") : TEXT("failed"), *Result.GetErrorRaw());
diff --git a/Engine/Plugins/Online/OnlineSubsystemEOS/Source/OnlineSubsystemEOS/Private/OnlineTitleFileEOS.cpp b/Engine/Plugins/Online/OnlineSubsystemEOS/Source/OnlineSubsystemEOS/Private/OnlineTitleFileEOS.cpp
--- a/Engine/Plugins/Online/OnlineSubsystemEOS/Source/OnlineSubsystemEOS/Private/OnlineTitleFileEOS.cpp
+++ b/Engine/Plugins/Online/OnlineSubsystemEOS/Source/OnlineSubsystemEOS/Private/OnlineTitleFileEOS.cpp
@@ -97,7 +97,7 @@ void FOnlineTitleFileEOS::DeleteCachedFiles(bool bSkipEnumerated)
 	EOS_TitleStorage_DeleteCacheOptions DeleteCacheOptions = { };
 	DeleteCacheOptions.ApiVersion = 1;
 	UE_EOS_CHECK_API_MISMATCH(EOS_TITLESTORAGE_DELETECACHEOPTIONS_API_LATEST, 1);
-	DeleteCacheOptions.LocalUserId = EOSSubsystem->UserManager->GetLocalProductUserId();	// Get a local user if one is available, but this is not required
+	DeleteCacheOptions.LocalUserId = EOSSubsystem->UserManager->GetLocalProductUserId(EOSSubsystem->UserManager->GetDefaultLocalUser());	// Get a local user if one is available, but this is not required
 
 	FDeleteCacheCompleteCallback* CallbackObj = new FDeleteCacheCompleteCallback(FOnlineTitleFileEOSWeakPtr(AsShared()));
 	CallbackObj->CallbackLambda = [this](const EOS_TitleStorage_DeleteCacheCallbackInfo* Data)
@@ -167,7 +167,7 @@ bool FOnlineTitleFileEOS::EnumerateFiles(const FPagedQuery& Page)
 	EOS_TitleStorage_QueryFileListOptions QueryFileListOptions = { };
 	QueryFileListOptions.ApiVersion = 1;
 	UE_EOS_CHECK_API_MISMATCH(EOS_TITLESTORAGE_QUERYFILELISTOPTIONS_API_LATEST, 1);
-	QueryFileListOptions.LocalUserId = EOSSubsystem->UserManager->GetLocalProductUserId();	// Get a local user if one is available, but this is not required
+	QueryFileListOptions.LocalUserId = EOSSubsystem->UserManager->GetLocalProductUserId(EOSSubsystem->UserManager->GetDefaultLocalUser());	// Get a local user if one is available, but this is not required
 	QueryFileListOptions.ListOfTags = const_cast<const char**>(AnsiTags.GetData());
 	QueryFileListOptions.ListOfTagsCount = AnsiTags.Num();
 
@@ -344,7 +344,7 @@ bool FOnlineTitleFileEOS::ReadFile(const FString& FileName)
 	EOS_TitleStorage_ReadFileOptions ReadFileOptions = { };
 	ReadFileOptions.ApiVersion = 1;
 	UE_EOS_CHECK_API_MISMATCH(EOS_TITLESTORAGE_READFILEOPTIONS_API_LATEST, 1);
-	ReadFileOptions.LocalUserId = EOSSubsystem->UserManager->GetLocalProductUserId();	// Get a local user if one is available, but this is not required
+	ReadFileOptions.LocalUserId = EOSSubsystem->UserManager->GetLocalProductUserId(EOSSubsystem->UserManager->GetDefaultLocalUser());	// Get a local user if one is available, but this is not required
 	ReadFileOptions.Filename = AnsiFileName;
 	ReadFileOptions.ReadChunkLengthBytes = (uint32)ReadChunkSize;
 	ReadFileOptions.ReadFileDataCallback = CallbackObj->GetNested1CallbackPtr();
diff --git a/Engine/Plugins/Online/OnlineSubsystemEOS/Source/OnlineSubsystemEOS/Private/SocketSubsystemEOSUtils_OnlineSubsystemEOS.cpp b/Engine/Plugins/Online/OnlineSubsystemEOS/Source/OnlineSubsystemEOS/Private/SocketSubsystemEOSUtils_OnlineSubsystemEOS.cpp
--- a/Engine/Plugins/Online/OnlineSubsystemEOS/Source/OnlineSubsystemEOS/Private/SocketSubsystemEOSUtils_OnlineSubsystemEOS.cpp
+++ b/Engine/Plugins/Online/OnlineSubsystemEOS/Source/OnlineSubsystemEOS/Private/SocketSubsystemEOSUtils_OnlineSubsystemEOS.cpp
@@ -19,7 +19,7 @@ EOS_ProductUserId FSocketSubsystemEOSUtils_OnlineSubsystemEOS::GetLocalUserId()
 {
 	EOS_ProductUserId Result = nullptr;
 
-	Result = SubsystemEOS.UserManager->GetLocalProductUserId();
+	Result = SubsystemEOS.UserManager->GetLocalProductUserId(SubsystemEOS.UserManager->GetDefaultLocalUser());
 
 	return Result;
 }
diff --git a/Engine/Plugins/Online/OnlineSubsystemEOS/Source/OnlineSubsystemEOS/Private/UserManagerEOS.cpp b/Engine/Plugins/Online/OnlineSubsystemEOS/Source/OnlineSubsystemEOS/Private/UserManagerEOS.cpp
--- a/Engine/Plugins/Online/OnlineSubsystemEOS/Source/OnlineSubsystemEOS/Private/UserManagerEOS.cpp
+++ b/Engine/Plugins/Online/OnlineSubsystemEOS/Source/OnlineSubsystemEOS/Private/UserManagerEOS.cpp
@@ -920,7 +920,7 @@ void FUserManagerEOS::FullLoginCallback(int32 LocalUserNum, EOS_EpicAccountId Ac
 			if (EpicAccountIdToOnlineUserMap.Contains(Data->PresenceUserId))
 			{
 				// Update the presence data to the most recent
-				UpdatePresence(Data->PresenceUserId);
+				UpdatePresence(LocalUserNum, Data->PresenceUserId);
 				return;
 			}
 		};
@@ -1154,11 +1154,6 @@ const FUniqueNetIdEOSPtr FUserManagerEOS::GetLocalUniqueNetIdEOS(int32 LocalUser
 	return nullptr;
 }
 
-const FUniqueNetIdEOSPtr FUserManagerEOS::GetLocalUniqueNetIdEOS() const
-{
-	return GetLocalUniqueNetIdEOS(DefaultLocalUser);
-}
-
 EOS_EpicAccountId FUserManagerEOS::GetLocalEpicAccountId(int32 LocalUserNum) const
 {
 	if (FUniqueNetIdEOSPtr UniqueNetId = GetLocalUniqueNetIdEOS(LocalUserNum))
@@ -1169,11 +1164,6 @@ EOS_EpicAccountId FUserManagerEOS::GetLocalEpicAccountId(int32 LocalUserNum) con
 	return nullptr;
 }
 
-EOS_EpicAccountId FUserManagerEOS::GetLocalEpicAccountId() const
-{
-	return GetLocalEpicAccountId(DefaultLocalUser);
-}
-
 EOS_ProductUserId FUserManagerEOS::GetLocalProductUserId(int32 LocalUserNum) const
 {
 	if (FUniqueNetIdEOSPtr UniqueNetId = GetLocalUniqueNetIdEOS(LocalUserNum))
@@ -1184,11 +1174,6 @@ EOS_ProductUserId FUserManagerEOS::GetLocalProductUserId(int32 LocalUserNum) con
 	return nullptr;
 }
 
-EOS_ProductUserId FUserManagerEOS::GetLocalProductUserId() const
-{
-	return GetLocalProductUserId(DefaultLocalUser);
-}
-
 int32 FUserManagerEOS::GetLocalUserNumFromUniqueNetId(const FUniqueNetId& NetId) const
 {
 	const FUniqueNetIdEOS& EosId = FUniqueNetIdEOS::Cast(NetId);
@@ -1296,7 +1281,7 @@ typedef TEOSCallback<EOS_Connect_OnQueryProductUserIdMappingsCallback, EOS_Conne
  *
  * @return true if the operation was successful, false otherwise
  */
-bool FUserManagerEOS::GetEpicAccountIdFromProductUserId(const EOS_ProductUserId& ProductUserId, EOS_EpicAccountId& OutEpicAccountId) const
+bool FUserManagerEOS::GetEpicAccountIdFromProductUserId(int32 LocalUserNum, const EOS_ProductUserId& ProductUserId, EOS_EpicAccountId& OutEpicAccountId) const
 {
 	bool bResult = false;
 
@@ -1307,7 +1292,7 @@ bool FUserManagerEOS::GetEpicAccountIdFromProductUserId(const EOS_ProductUserId&
 	Options.ApiVersion = 1;
 	UE_EOS_CHECK_API_MISMATCH(EOS_CONNECT_GETPRODUCTUSERIDMAPPING_API_LATEST, 1);
 	Options.AccountIdType = EOS_EExternalAccountType::EOS_EAT_EPIC;
-	Options.LocalUserId = GetLocalProductUserId();
+	Options.LocalUserId = GetLocalProductUserId(LocalUserNum);
 	Options.TargetProductUserId = ProductUserId;
 
 	EOS_EResult Result = EOS_Connect_GetProductUserIdMapping(EOSSubsystem->ConnectHandle, &Options, EpicIdStr, &EpicIdStrSize);
@@ -1324,7 +1309,7 @@ bool FUserManagerEOS::GetEpicAccountIdFromProductUserId(const EOS_ProductUserId&
 	return bResult;
 }
 
-void FUserManagerEOS::ResolveUniqueNetId(const EOS_ProductUserId& ProductUserId, const FResolveUniqueNetIdCallback& Callback) const
+void FUserManagerEOS::ResolveUniqueNetId(int32 LocalUserNum, const EOS_ProductUserId& ProductUserId, const FResolveUniqueNetIdCallback& Callback) const
 {
 	TArray<EOS_ProductUserId> ProductUserIds = { const_cast<EOS_ProductUserId>(ProductUserId) };
 
@@ -1332,10 +1317,10 @@ void FUserManagerEOS::ResolveUniqueNetId(const EOS_ProductUserId& ProductUserId,
 		OriginalCallback(ResolvedUniqueNetIds[ProductUserId]);
 	};
 
-	ResolveUniqueNetIds(ProductUserIds, GroupCallback);
+	ResolveUniqueNetIds(LocalUserNum, ProductUserIds, GroupCallback);
 }
 
-void FUserManagerEOS::ResolveUniqueNetIds(const TArray<EOS_ProductUserId>& ProductUserIds, const FResolveUniqueNetIdsCallback& Callback) const
+void FUserManagerEOS::ResolveUniqueNetIds(int32 LocalUserNum, const TArray<EOS_ProductUserId>& ProductUserIds, const FResolveUniqueNetIdsCallback& Callback) const
 {
 	TMap<EOS_ProductUserId, FUniqueNetIdEOSRef> ResolvedUniqueNetIds;
 	TArray<EOS_ProductUserId> ProductUserIdsToResolve;
@@ -1345,7 +1330,7 @@ void FUserManagerEOS::ResolveUniqueNetIds(const TArray<EOS_ProductUserId>& Produ
 		EOS_EpicAccountId EpicAccountId;
 
 		// We check first if the Product User Id has already been queried, which would allow us to retrieve its Epic Account Id directly
-		if (GetEpicAccountIdFromProductUserId(ProductUserId, EpicAccountId))
+		if (GetEpicAccountIdFromProductUserId(LocalUserNum, ProductUserId, EpicAccountId))
 		{
 			const FUniqueNetIdEOSRef UniqueNetId = FUniqueNetIdEOSRegistry::FindOrAdd(EpicAccountId, ProductUserId).ToSharedRef();
 
@@ -1368,7 +1353,7 @@ void FUserManagerEOS::ResolveUniqueNetIds(const TArray<EOS_ProductUserId>& Produ
 		QueryProductUserIdMappingsOptions.ProductUserIdCount = ProductUserIdsToResolve.Num();
 
 		FConnectQueryProductUserIdMappingsCallback* CallbackObj = new FConnectQueryProductUserIdMappingsCallback(FUserManagerEOSConstWeakPtr(AsShared()));
-		CallbackObj->CallbackLambda = [this, ProductUserIdsToResolve, ResolvedUniqueNetIds = MoveTemp(ResolvedUniqueNetIds), Callback](const EOS_Connect_QueryProductUserIdMappingsCallbackInfo* Data) mutable
+		CallbackObj->CallbackLambda = [this, LocalUserNum, ProductUserIdsToResolve, ResolvedUniqueNetIds = MoveTemp(ResolvedUniqueNetIds), Callback](const EOS_Connect_QueryProductUserIdMappingsCallbackInfo* Data) mutable
 		{
 			if (Data->ResultCode != EOS_EResult::EOS_Success)
 			{
@@ -1379,7 +1364,7 @@ void FUserManagerEOS::ResolveUniqueNetIds(const TArray<EOS_ProductUserId>& Produ
 			{
 				EOS_EpicAccountId EpicAccountId = nullptr;
 
-				GetEpicAccountIdFromProductUserId(ProductUserId, EpicAccountId);
+				GetEpicAccountIdFromProductUserId(LocalUserNum, ProductUserId, EpicAccountId);
 
 				const FUniqueNetIdEOSRef UniqueNetId = FUniqueNetIdEOSRegistry::FindOrAdd(EpicAccountId, ProductUserId).ToSharedRef();
 
@@ -2610,7 +2595,7 @@ void FUserManagerEOS::QueryPresence(const FUniqueNetId& UserId, const FOnPresenc
 	EOS_Presence_HasPresenceOptions HasOptions = { };
 	HasOptions.ApiVersion = 1;
 	UE_EOS_CHECK_API_MISMATCH(EOS_PRESENCE_HASPRESENCE_API_LATEST, 1);
-	HasOptions.LocalUserId = GetLocalEpicAccountId();
+	HasOptions.LocalUserId = GetLocalEpicAccountId(DefaultLocalUser);
 	HasOptions.TargetUserId = AccountId;
 	EOS_Bool bHasPresence = EOS_Presence_HasPresence(EOSSubsystem->PresenceHandle, &HasOptions);
 	if (bHasPresence == EOS_FALSE)
@@ -2621,7 +2606,7 @@ void FUserManagerEOS::QueryPresence(const FUniqueNetId& UserId, const FOnPresenc
 			if (Data->ResultCode == EOS_EResult::EOS_Success && EpicAccountIdToOnlineUserMap.Contains(Data->TargetUserId))
 			{
 				// Update the presence data to the most recent
-				UpdatePresence(Data->TargetUserId);
+				UpdatePresence(DefaultLocalUser, Data->TargetUserId);
 				FOnlineUserPtr OnlineUser = EpicAccountIdToOnlineUserMap[Data->TargetUserId];
 				Delegate.ExecuteIfBound(*OnlineUser->GetUserId(), true);
 				return;
@@ -2642,18 +2627,18 @@ void FUserManagerEOS::QueryPresence(const FUniqueNetId& UserId, const FOnPresenc
 	}
 
 	// Update the presence data to the most recent
-	UpdatePresence(HasOptions.TargetUserId);
+	UpdatePresence(DefaultLocalUser, HasOptions.TargetUserId);
 	// It's already present so trigger that it's done
 	Delegate.ExecuteIfBound(UserId, true);
 }
 
-void FUserManagerEOS::UpdatePresence(EOS_EpicAccountId AccountId)
+void FUserManagerEOS::UpdatePresence(int32 LocalUserNum, EOS_EpicAccountId AccountId)
 {
 	EOS_Presence_Info* PresenceInfo = nullptr;
 	EOS_Presence_CopyPresenceOptions Options = { };
 	Options.ApiVersion = 3;
 	UE_EOS_CHECK_API_MISMATCH(EOS_PRESENCE_COPYPRESENCE_API_LATEST, 3);
-	Options.LocalUserId = GetLocalEpicAccountId();
+	Options.LocalUserId = GetLocalEpicAccountId(LocalUserNum);
 	Options.TargetUserId = AccountId;
 	EOS_EResult CopyResult = EOS_Presence_CopyPresence(EOSSubsystem->PresenceHandle, &Options, &PresenceInfo);
 	if (CopyResult == EOS_EResult::EOS_Success)
diff --git a/Engine/Plugins/Online/OnlineSubsystemEOS/Source/OnlineSubsystemEOS/Private/UserManagerEOS.h b/Engine/Plugins/Online/OnlineSubsystemEOS/Source/OnlineSubsystemEOS/Private/UserManagerEOS.h
--- a/Engine/Plugins/Online/OnlineSubsystemEOS/Source/OnlineSubsystemEOS/Private/UserManagerEOS.h
+++ b/Engine/Plugins/Online/OnlineSubsystemEOS/Source/OnlineSubsystemEOS/Private/UserManagerEOS.h
@@ -285,11 +285,8 @@ public:
 // ~IOnlineUser Interface
 
 	const FUniqueNetIdEOSPtr GetLocalUniqueNetIdEOS(int32 LocalUserNum) const;
-	const FUniqueNetIdEOSPtr GetLocalUniqueNetIdEOS() const;
 	EOS_EpicAccountId GetLocalEpicAccountId(int32 LocalUserNum) const;
-	EOS_EpicAccountId GetLocalEpicAccountId() const;
 	EOS_ProductUserId GetLocalProductUserId(int32 LocalUserNum) const;
-	EOS_ProductUserId GetLocalProductUserId() const;
 
 	int32 GetLocalUserNumFromUniqueNetId(const FUniqueNetId& NetId) const;
 	int32 GetLocalUserNumFromEpicAccountId(const EOS_EpicAccountId& EpicAccountId) const;
@@ -304,9 +301,9 @@ public:
 
 	typedef TFunction<void(TMap<EOS_ProductUserId, FUniqueNetIdEOSRef> ResolvedUniqueNetIds)> FResolveUniqueNetIdsCallback;
 	typedef TFunction<void(FUniqueNetIdEOSRef ResolvedUniqueNetId)> FResolveUniqueNetIdCallback;
-	bool GetEpicAccountIdFromProductUserId(const EOS_ProductUserId& ProductUserId, EOS_EpicAccountId& OutEpicAccountId) const;
-	void ResolveUniqueNetId(const EOS_ProductUserId& ProductUserId, const FResolveUniqueNetIdCallback& Callback) const;
-	void ResolveUniqueNetIds(const TArray<EOS_ProductUserId>& ProductUserIds, const FResolveUniqueNetIdsCallback& Callback) const;
+	bool GetEpicAccountIdFromProductUserId(int32 LocalUserNum, const EOS_ProductUserId& ProductUserId, EOS_EpicAccountId& OutEpicAccountId) const;
+	void ResolveUniqueNetId(int32 LocalUserNum, const EOS_ProductUserId& ProductUserId, const FResolveUniqueNetIdCallback& Callback) const;
+	void ResolveUniqueNetIds(int32 LocalUserNum, const TArray<EOS_ProductUserId>& ProductUserIds, const FResolveUniqueNetIdsCallback& Callback) const;
 
 	FOnlineUserPtr GetLocalOnlineUser(int32 LocalUserNum) const;
 	FOnlineUserPtr GetOnlineUser(EOS_ProductUserId UserId) const;
@@ -346,7 +343,7 @@ private:
 	bool IsFriendQueryUserInfoOngoing(int32 LocalUserNum);
 	void ProcessReadFriendsListComplete(int32 LocalUserNum, bool bWasSuccessful, const FString& ErrorStr);
 
-	void UpdatePresence(EOS_EpicAccountId AccountId);
+	void UpdatePresence(int32 LocalUserNum, EOS_EpicAccountId AccountId);
 	void UpdateFriendPresence(const FString& FriendId, FOnlineUserPresenceRef Presence);
 
 	void GetPlatformAuthToken(int32 LocalUserNum, const FOnGetLinkedAccountAuthTokenCompleteDelegate& Delegate) const;
