- Demote warning to Log to reduce log noise

diff --git a/Engine/Plugins/Online/OnlineSubsystemEOS/Source/OnlineSubsystemEOS/Private/EOSSettings.cpp b/Engine/Plugins/Online/OnlineSubsystemEOS/Source/OnlineSubsystemEOS/Private/EOSSettings.cpp
--- a/Engine/Plugins/Online/OnlineSubsystemEOS/Source/OnlineSubsystemEOS/Private/EOSSettings.cpp
+++ b/Engine/Plugins/Online/OnlineSubsystemEOS/Source/OnlineSubsystemEOS/Private/EOSSettings.cpp
@@ -217,7 +217,7 @@ bool UEOSSettings::GetSelectedArtifactSettings(FEOSArtifactSettings& OutSettings
 			return true;
 		}
 
-		UE_LOG_ONLINE(Warning, TEXT("UEOSSettings::GetSelectedArtifactSettings() ArtifactName=[%s] SandboxId=[%s] no settings found."), *ArtifactName, *SandboxId);
+		UE_LOG_ONLINE(Log, TEXT("UEOSSettings::GetSelectedArtifactSettings() ArtifactName=[%s] SandboxId=[%s] no settings found for pair, falling back on just ArtifactName."), *ArtifactName, *SandboxId);
 	}
 
 	// Fall back on just matching the Artifact name. This assumes non-EGS and only one settings entry per ArtifactName in config.
