- Move SDK Init from ModuleInit to instance init

diff --git a/Engine/Plugins/Online/OnlineSubsystemEOS/Source/OnlineSubsystemEOS/Private/OnlineSubsystemEOS.cpp b/Engine/Plugins/Online/OnlineSubsystemEOS/Source/OnlineSubsystemEOS/Private/OnlineSubsystemEOS.cpp
--- a/Engine/Plugins/Online/OnlineSubsystemEOS/Source/OnlineSubsystemEOS/Private/OnlineSubsystemEOS.cpp
+++ b/Engine/Plugins/Online/OnlineSubsystemEOS/Source/OnlineSubsystemEOS/Private/OnlineSubsystemEOS.cpp
@@ -18,9 +18,9 @@
 
 #include "Features/IModularFeatures.h"
 #include "Misc/App.h"
-#include "Misc/NetworkVersion.h"
-#include "Misc/App.h"
 #include "Misc/ConfigCacheIni.h"
+#include "Misc/Fork.h"
+#include "Misc/NetworkVersion.h"
 
 #if WITH_EOS_SDK
 
@@ -144,11 +144,15 @@ void FOnlineSubsystemEOS::ModuleInit()
 		return;
 	}
 
-	EOS_EResult InitResult = EOSSDKManager->Initialize();
-	if (InitResult != EOS_EResult::EOS_Success)
+	// If a fork is requested, we need to wait for post-fork to init the SDK
+	if (!FForkProcessHelper::IsForkRequested() || FForkProcessHelper::IsForkedChildProcess())
 	{
-		UE_LOG_ONLINE(Error, TEXT("FOnlineSubsystemEOS: failed to initialize the EOS SDK with result code (%d)"), InitResult);
-		return;
+		EOS_EResult InitResult = EOSSDKManager->Initialize();
+		if (InitResult != EOS_EResult::EOS_Success)
+		{
+			UE_LOG_ONLINE(Error, TEXT("FOnlineSubsystemEOS: failed to initialize the EOS SDK with result code (%d)"), InitResult);
+			return;
+		}
 	}
 }
 
@@ -277,6 +281,13 @@ bool FOnlineSubsystemEOS::Init()
 		return false;
 	}
 
+	EOS_EResult InitResult = EOSSDKManager->Initialize();
+	if (InitResult != EOS_EResult::EOS_Success)
+	{
+		UE_LOG_ONLINE(Error, TEXT("FOnlineSubsystemEOS: failed to initialize the EOS SDK with result code (%d)"), InitResult);
+		return false;
+	}
+
 	if (!PlatformCreate())
 	{
 		return false;
