- Remove unnecessary "IsDS" cached bool, and just use IsRunningDedicatedServer()

diff --git a/Engine/Plugins/Online/OnlineSubsystemEOS/Source/OnlineSubsystemEOS/Private/OnlineSessionEOS.cpp b/Engine/Plugins/Online/OnlineSubsystemEOS/Source/OnlineSubsystemEOS/Private/OnlineSessionEOS.cpp
--- a/Engine/Plugins/Online/OnlineSubsystemEOS/Source/OnlineSubsystemEOS/Private/OnlineSessionEOS.cpp
+++ b/Engine/Plugins/Online/OnlineSubsystemEOS/Source/OnlineSubsystemEOS/Private/OnlineSessionEOS.cpp
@@ -500,7 +500,6 @@ void FOnlineSessionEOS::Init()
 	LobbyHandle = EOS_Platform_GetLobbyInterface(*EOSSubsystem->EOSPlatformHandle);
 	RegisterLobbyNotifications();
 
-	bIsDedicatedServer = IsRunningDedicatedServer();
 	bIsUsingP2PSockets = false;
 
 	// Presence usage used to be guessed from other session parameters like permission level, with this new logic we have a reliable way of transmitting that information, in the form of a session attribute
@@ -991,7 +990,7 @@ bool FOnlineSessionEOS::CreateSession(int32 HostingPlayerNum, FName SessionName,
 	FNamedOnlineSession* Session = GetNamedSession(SessionName);
 	if (Session == nullptr)
 	{
-		if (bIsDedicatedServer || EOSSubsystem->UserManager->GetLoginStatus(HostingPlayerNum) >= ELoginStatus::UsingLocalProfile)
+		if (IsRunningDedicatedServer() || EOSSubsystem->UserManager->GetLoginStatus(HostingPlayerNum) >= ELoginStatus::UsingLocalProfile)
 		{
 			// Create a new session and deep copy the game settings
 			Session = AddNamedSession(SessionName, NewSessionSettings);
@@ -1001,7 +1000,7 @@ bool FOnlineSessionEOS::CreateSession(int32 HostingPlayerNum, FName SessionName,
 			Session->OwningUserId = EOSSubsystem->UserManager->GetUniquePlayerId(HostingPlayerNum);
 			Session->OwningUserName = EOSSubsystem->UserManager->GetPlayerNickname(HostingPlayerNum);
 
-			if (bIsDedicatedServer || (Session->OwningUserId.IsValid() && Session->OwningUserId->IsValid()))
+			if (IsRunningDedicatedServer() || (Session->OwningUserId.IsValid() && Session->OwningUserId->IsValid()))
 			{
 				// RegisterPlayer will update these values for the local player
 				Session->NumOpenPrivateConnections = NewSessionSettings.NumPrivateConnections;
@@ -1356,7 +1355,7 @@ uint32 FOnlineSessionEOS::CreateEOSSession(int32 HostingPlayerNum, FNamedOnlineS
 
 	FString HostAddr;
 	// If we are not a dedicated server and are using p2p sockets, then we need to add a custom URL for connecting
-	if (!bIsDedicatedServer && bIsUsingP2PSockets)
+	if (!IsRunningDedicatedServer() && bIsUsingP2PSockets)
 	{
 		// Because some platforms remap ports, we will use the ID of the name of the net driver to be our port instead
 		FName NetDriverName = GetDefault<UNetDriverEOS>()->NetDriverName;
diff --git a/Engine/Plugins/Online/OnlineSubsystemEOS/Source/OnlineSubsystemEOS/Private/OnlineSessionEOS.h b/Engine/Plugins/Online/OnlineSubsystemEOS/Source/OnlineSubsystemEOS/Private/OnlineSessionEOS.h
--- a/Engine/Plugins/Online/OnlineSubsystemEOS/Source/OnlineSubsystemEOS/Private/OnlineSessionEOS.h
+++ b/Engine/Plugins/Online/OnlineSubsystemEOS/Source/OnlineSubsystemEOS/Private/OnlineSessionEOS.h
@@ -344,7 +344,6 @@ private:
 	EOS_NotificationId SessionInviteAcceptedId;
 	FCallbackBase* SessionInviteAcceptedCallback;
 
-	bool bIsDedicatedServer;
 	bool bIsUsingP2PSockets;
 	bool bUsePresenceAttribute = false;
 
