- Fix double lookup in TOnlinePlayerList::Remove

diff --git a/Engine/Plugins/Online/OnlineSubsystemEOS/Source/OnlineSubsystemEOS/Private/UserManagerEOS.h b/Engine/Plugins/Online/OnlineSubsystemEOS/Source/OnlineSubsystemEOS/Private/UserManagerEOS.h
--- a/Engine/Plugins/Online/OnlineSubsystemEOS/Source/OnlineSubsystemEOS/Private/UserManagerEOS.h
+++ b/Engine/Plugins/Online/OnlineSubsystemEOS/Source/OnlineSubsystemEOS/Private/UserManagerEOS.h
@@ -106,11 +106,7 @@ public:
 
 	void Remove(const FString& InNetId, ListClass InListEntry)
 	{
-		const ListClass* Found = NetIdStringToListEntryMap.Find(InNetId);
-		if (Found != nullptr)
-		{
-			NetIdStringToListEntryMap.Remove(InNetId);
-		}
+		NetIdStringToListEntryMap.Remove(InNetId);
 		ListEntries.Remove(InListEntry);
 	}
 
