UserManagerEOS internal data structures refactor (Part 2)
- Added Find and FindChecked methods to EOS NetId Registry, refactored FindOrAdd to return a SharedRef instead of a SharedPtr.
- Removed redundant NetId and OnlineUser maps, replaced functionality with existing maps and registry lookups

diff --git a/Engine/Plugins/Online/OnlineSubsystemEOS/Source/OnlineSubsystemEOS/Private/OnlineLeaderboardsEOS.cpp b/Engine/Plugins/Online/OnlineSubsystemEOS/Source/OnlineSubsystemEOS/Private/OnlineLeaderboardsEOS.cpp
--- a/Engine/Plugins/Online/OnlineSubsystemEOS/Source/OnlineSubsystemEOS/Private/OnlineLeaderboardsEOS.cpp
+++ b/Engine/Plugins/Online/OnlineSubsystemEOS/Source/OnlineSubsystemEOS/Private/OnlineLeaderboardsEOS.cpp
@@ -316,7 +316,7 @@ bool FOnlineLeaderboardsEOS::ReadLeaderboardsAroundRank(int32 Rank, uint32 Range
 						Nickname = TEXT("Unknown Player");
 					}
 					const EOS_EpicAccountId AccountId = EOS_EpicAccountId_FromString(EpicIdStr);
-					const FUniqueNetIdEOSRef NetId = FUniqueNetIdEOSRegistry::FindOrAdd(AccountId, Record->UserId).ToSharedRef();
+					const FUniqueNetIdEOSRef NetId = FUniqueNetIdEOSRegistry::FindOrAdd(AccountId, Record->UserId);
 
 					FOnlineStatsRow* Row = new(LambdaReadObject->Rows) FOnlineStatsRow(Nickname, NetId);
 					Row->Rank = Record->Rank;
diff --git a/Engine/Plugins/Online/OnlineSubsystemEOS/Source/OnlineSubsystemEOS/Private/OnlineSubsystemEOSTypes.cpp b/Engine/Plugins/Online/OnlineSubsystemEOS/Source/OnlineSubsystemEOS/Private/OnlineSubsystemEOSTypes.cpp
--- a/Engine/Plugins/Online/OnlineSubsystemEOS/Source/OnlineSubsystemEOS/Private/OnlineSubsystemEOSTypes.cpp
+++ b/Engine/Plugins/Online/OnlineSubsystemEOS/Source/OnlineSubsystemEOS/Private/OnlineSubsystemEOSTypes.cpp
@@ -101,7 +101,7 @@ FUniqueNetIdEOSRegistry& FUniqueNetIdEOSRegistry::Get()
 	return TLazySingleton<FUniqueNetIdEOSRegistry>::Get();
 }
 
-FUniqueNetIdEOSPtr FUniqueNetIdEOSRegistry::FindOrAddImpl(const FString& NetIdStr)
+FUniqueNetIdEOSRef FUniqueNetIdEOSRegistry::FindOrAddImpl(const FString& NetIdStr)
 {
 	FString EpicAccountIdStr;
 	FString ProductUserIdStr;
@@ -109,7 +109,7 @@ FUniqueNetIdEOSPtr FUniqueNetIdEOSRegistry::FindOrAddImpl(const FString& NetIdSt
 		|| (!EpicAccountIdStr.IsEmpty() && EpicAccountIdStr.Len() != EOS_EPICACCOUNTID_MAX_LENGTH)
 		|| (!ProductUserIdStr.IsEmpty() && ProductUserIdStr.Len() != EOS_PRODUCTUSERID_MAX_LENGTH))
 	{
-		return nullptr;
+		return FUniqueNetIdEOS::EmptyId();
 	}
 
 	EOS_EpicAccountId EpicAccountId = nullptr;
@@ -126,7 +126,7 @@ FUniqueNetIdEOSPtr FUniqueNetIdEOSRegistry::FindOrAddImpl(const FString& NetIdSt
 	return FindOrAddImpl(EpicAccountId, ProductUserId);
 }
 
-FUniqueNetIdEOSPtr FUniqueNetIdEOSRegistry::FindOrAddImpl(const uint8* Bytes, int32 Size)
+FUniqueNetIdEOSRef FUniqueNetIdEOSRegistry::FindOrAddImpl(const uint8* Bytes, int32 Size)
 {
 	if (Size == EOS_ID_BYTE_SIZE)
 	{
@@ -134,12 +134,13 @@ FUniqueNetIdEOSPtr FUniqueNetIdEOSRegistry::FindOrAddImpl(const uint8* Bytes, in
 		FUniqueNetIdEOS Temp(Bytes, Size);
 		return FindOrAddImpl(Temp.GetEpicAccountId(), Temp.GetProductUserId());
 	}
-	return nullptr;
+	return FUniqueNetIdEOS::EmptyId();
 }
 
-FUniqueNetIdEOSPtr FUniqueNetIdEOSRegistry::FindOrAddImpl(const EOS_EpicAccountId InEpicAccountId, const EOS_ProductUserId InProductUserId)
+FUniqueNetIdEOSRef FUniqueNetIdEOSRegistry::FindOrAddImpl(const EOS_EpicAccountId InEpicAccountId, const EOS_ProductUserId InProductUserId)
 {
 	FUniqueNetIdEOSPtr Result;
+
 	const bool bInEpicAccountIdValid = EOS_EpicAccountId_IsValid(InEpicAccountId) == EOS_TRUE;
 	const bool bInProductUserIdValid = EOS_ProductUserId_IsValid(InProductUserId) == EOS_TRUE;
 	if (bInEpicAccountIdValid || bInProductUserIdValid)
@@ -150,6 +151,7 @@ FUniqueNetIdEOSPtr FUniqueNetIdEOSRegistry::FindOrAddImpl(const EOS_EpicAccountI
 		auto FindExisting = [this, InEpicAccountId, InProductUserId, bInEpicAccountIdValid, bInProductUserIdValid, &bUpdateEpicAccountId, &bUpdateProductUserId]()
 		{
 			FUniqueNetIdEOSPtr Result;
+
 			if (const FUniqueNetIdEOSRef* FoundEas = bInEpicAccountIdValid ? EasToNetId.Find(InEpicAccountId) : nullptr)
 			{
 				Result = *FoundEas;
@@ -171,12 +173,11 @@ FUniqueNetIdEOSPtr FUniqueNetIdEOSRegistry::FindOrAddImpl(const EOS_EpicAccountI
 				check(!bFoundProductUserIdValid || !bInProductUserIdValid || InProductUserId == FoundProductUserId);
 				bUpdateEpicAccountId = !bFoundEpicAccountIdValid && bInEpicAccountIdValid;
 				bUpdateProductUserId = !bFoundProductUserIdValid && bInProductUserIdValid;
-
 			}
+
 			return Result;
 		};
 
-
 		{
 			// First take read lock and look for existing elements
 			const FReadScopeLock ReadLock(Lock);
@@ -222,5 +223,49 @@ FUniqueNetIdEOSPtr FUniqueNetIdEOSRegistry::FindOrAddImpl(const EOS_EpicAccountI
 		}
 	}
 
-	return Result;
+	return Result ? Result.ToSharedRef() : FUniqueNetIdEOS::EmptyId();
+}
+
+FUniqueNetIdEOSPtr FUniqueNetIdEOSRegistry::FindImpl(EOS_EpicAccountId EpicAccountId)
+{
+	FUniqueNetIdEOSRef* Result = EasToNetId.Find(EpicAccountId);
+
+	return Result ? FUniqueNetIdEOSPtr(*Result) : nullptr;
+}
+
+FUniqueNetIdEOSPtr FUniqueNetIdEOSRegistry::FindImpl(EOS_ProductUserId ProductUserId)
+{
+	FUniqueNetIdEOSRef* Result = PuidToNetId.Find(ProductUserId);
+
+	return Result ? FUniqueNetIdEOSPtr(*Result) : nullptr;
+}
+
+FUniqueNetIdEOSRef FUniqueNetIdEOSRegistry::FindCheckedImpl(EOS_EpicAccountId EpicAccountId)
+{
+	FUniqueNetIdEOSRef* Result = EasToNetId.Find(EpicAccountId);
+	check(Result != nullptr);
+
+	if (Result != nullptr)
+	{
+		return *Result;
+	}
+	else
+	{
+		return FUniqueNetIdEOS::EmptyId();
+	}
 }
+
+FUniqueNetIdEOSRef FUniqueNetIdEOSRegistry::FindCheckedImpl(EOS_ProductUserId ProductUserId)
+{
+	FUniqueNetIdEOSRef* Result = PuidToNetId.Find(ProductUserId);
+	check(Result != nullptr);
+
+	if (Result != nullptr)
+	{
+		return *Result;
+	}
+	else
+	{
+		return FUniqueNetIdEOS::EmptyId();
+	}
+}
\ No newline at end of file
diff --git a/Engine/Plugins/Online/OnlineSubsystemEOS/Source/OnlineSubsystemEOS/Private/OnlineSubsystemEOSTypes.h b/Engine/Plugins/Online/OnlineSubsystemEOS/Source/OnlineSubsystemEOS/Private/OnlineSubsystemEOSTypes.h
--- a/Engine/Plugins/Online/OnlineSubsystemEOS/Source/OnlineSubsystemEOS/Private/OnlineSubsystemEOSTypes.h
+++ b/Engine/Plugins/Online/OnlineSubsystemEOS/Source/OnlineSubsystemEOS/Private/OnlineSubsystemEOSTypes.h
@@ -79,9 +79,15 @@ private:
 class FUniqueNetIdEOSRegistry
 {
 public:
-	static FUniqueNetIdEOSPtr FindOrAdd(const FString& NetIdStr) { return Get().FindOrAddImpl(NetIdStr); }
-	static FUniqueNetIdEOSPtr FindOrAdd(const uint8* Bytes, int32 Size) { return Get().FindOrAddImpl(Bytes, Size); }
-	static FUniqueNetIdEOSPtr FindOrAdd(EOS_EpicAccountId EpicAccountId, EOS_ProductUserId ProductUserId) { return Get().FindOrAddImpl(EpicAccountId, ProductUserId); }
+	static FUniqueNetIdEOSRef FindOrAdd(const FString& NetIdStr) { return Get().FindOrAddImpl(NetIdStr); }
+	static FUniqueNetIdEOSRef FindOrAdd(const uint8* Bytes, int32 Size) { return Get().FindOrAddImpl(Bytes, Size); }
+	static FUniqueNetIdEOSRef FindOrAdd(EOS_EpicAccountId EpicAccountId, EOS_ProductUserId ProductUserId) { return Get().FindOrAddImpl(EpicAccountId, ProductUserId); }
+
+	static FUniqueNetIdEOSPtr Find(EOS_EpicAccountId EpicAccountId) { return Get().FindImpl(EpicAccountId); }
+	static FUniqueNetIdEOSPtr Find(EOS_ProductUserId ProductUserId) { return Get().FindImpl(ProductUserId); }
+
+	static FUniqueNetIdEOSRef FindChecked(EOS_EpicAccountId EpicAccountId) { return Get().FindCheckedImpl(EpicAccountId); }
+	static FUniqueNetIdEOSRef FindChecked(EOS_ProductUserId ProductUserId) { return Get().FindCheckedImpl(ProductUserId); }
 
 private:
 	FRWLock Lock;
@@ -90,9 +96,15 @@ private:
 
 	static FUniqueNetIdEOSRegistry& Get();
 
-	FUniqueNetIdEOSPtr FindOrAddImpl(const FString& NetIdStr);
-	FUniqueNetIdEOSPtr FindOrAddImpl(const uint8* Bytes, int32 Size);
-	FUniqueNetIdEOSPtr FindOrAddImpl(EOS_EpicAccountId EpicAccountId, EOS_ProductUserId ProductUserId);
+	FUniqueNetIdEOSRef FindOrAddImpl(const FString& NetIdStr);
+	FUniqueNetIdEOSRef FindOrAddImpl(const uint8* Bytes, int32 Size);
+	FUniqueNetIdEOSRef FindOrAddImpl(EOS_EpicAccountId EpicAccountId, EOS_ProductUserId ProductUserId);
+
+	FUniqueNetIdEOSPtr FindImpl(EOS_EpicAccountId EpicAccountId);
+	FUniqueNetIdEOSPtr FindImpl(EOS_ProductUserId ProductUserId);
+
+	FUniqueNetIdEOSRef FindCheckedImpl(EOS_EpicAccountId EpicAccountId);
+	FUniqueNetIdEOSRef FindCheckedImpl(EOS_ProductUserId ProductUserId);
 };
 
 #ifndef AUTH_ATTR_REFRESH_TOKEN
diff --git a/Engine/Plugins/Online/OnlineSubsystemEOS/Source/OnlineSubsystemEOS/Private/UserManagerEOS.cpp b/Engine/Plugins/Online/OnlineSubsystemEOS/Source/OnlineSubsystemEOS/Private/UserManagerEOS.cpp
--- a/Engine/Plugins/Online/OnlineSubsystemEOS/Source/OnlineSubsystemEOS/Private/UserManagerEOS.cpp
+++ b/Engine/Plugins/Online/OnlineSubsystemEOS/Source/OnlineSubsystemEOS/Private/UserManagerEOS.cpp
@@ -917,7 +917,7 @@ void FUserManagerEOS::FullLoginCallback(int32 LocalUserNum, EOS_EpicAccountId Ac
 		PresenceNotificationCallback = CallbackObj;
 		CallbackObj->CallbackLambda = [LocalUserNum, this](const EOS_Presence_PresenceChangedCallbackInfo* Data)
 		{
-			if (EpicAccountIdToOnlineUserMap.Contains(Data->PresenceUserId))
+			if (FUniqueNetIdEOSRegistry::Find(Data->PresenceUserId))
 			{
 				// Update the presence data to the most recent
 				UpdatePresence(LocalUserNum, Data->PresenceUserId);
@@ -1039,15 +1039,12 @@ void FUserManagerEOS::AddLocalUser(int32 LocalUserNum, EOS_EpicAccountId EpicAcc
 		DefaultLocalUser = LocalUserNum;
 	}
 
-	FUniqueNetIdEOSRef UserNetId = FUniqueNetIdEOSRegistry::FindOrAdd(EpicAccountId, UserId).ToSharedRef();
+	FUniqueNetIdEOSRef UserNetId = FUniqueNetIdEOSRegistry::FindOrAdd(EpicAccountId, UserId);
 	const FString& NetId = UserNetId->ToString();
 	FUserOnlineAccountEOSRef UserAccountRef(new FUserOnlineAccountEOS(UserNetId));
 
 	LocalUserUniqueNetIds.Insert(LocalUserNum, UserNetId);
 	UniqueNetIdToUserAccountMap.Emplace(UserNetId, UserAccountRef);
-	NetIdStringToOnlineUserMap.Emplace(*NetId, UserAccountRef);
-	AccountIdToStringMap.Emplace(EpicAccountId, NetId);
-	ProductUserIdToStringMap.Emplace(UserId, *NetId);
 	EpicAccountIdToAttributeAccessMap.Emplace(EpicAccountId, UserAccountRef);
 
 	// Init player lists
@@ -1332,7 +1329,7 @@ void FUserManagerEOS::ResolveUniqueNetIds(int32 LocalUserNum, const TArray<EOS_P
 		// We check first if the Product User Id has already been queried, which would allow us to retrieve its Epic Account Id directly
 		if (GetEpicAccountIdFromProductUserId(LocalUserNum, ProductUserId, EpicAccountId))
 		{
-			const FUniqueNetIdEOSRef UniqueNetId = FUniqueNetIdEOSRegistry::FindOrAdd(EpicAccountId, ProductUserId).ToSharedRef();
+			const FUniqueNetIdEOSRef UniqueNetId = FUniqueNetIdEOSRegistry::FindOrAdd(EpicAccountId, ProductUserId);
 
 			ResolvedUniqueNetIds.Add(ProductUserId, UniqueNetId);
 		}
@@ -1366,7 +1363,7 @@ void FUserManagerEOS::ResolveUniqueNetIds(int32 LocalUserNum, const TArray<EOS_P
 
 				GetEpicAccountIdFromProductUserId(LocalUserNum, ProductUserId, EpicAccountId);
 
-				const FUniqueNetIdEOSRef UniqueNetId = FUniqueNetIdEOSRegistry::FindOrAdd(EpicAccountId, ProductUserId).ToSharedRef();
+				const FUniqueNetIdEOSRef UniqueNetId = FUniqueNetIdEOSRegistry::FindOrAdd(EpicAccountId, ProductUserId);
 
 				ResolvedUniqueNetIds.Add(ProductUserId, UniqueNetId);
 			}
@@ -1385,43 +1382,45 @@ void FUserManagerEOS::ResolveUniqueNetIds(int32 LocalUserNum, const TArray<EOS_P
 FOnlineUserPtr FUserManagerEOS::GetLocalOnlineUser(int32 LocalUserNum) const
 {
 	FOnlineUserPtr OnlineUser;
-	if (LocalUserUniqueNetIds.IsValidIndex(LocalUserNum))
-	{
-		const FUniqueNetIdEOSRef& NetId = LocalUserUniqueNetIds[LocalUserNum];
 
-		if (NetIdStringToOnlineUserMap.Contains(NetId->ToString()))
+	if (const FUniqueNetIdEOSPtr NetIdPtr = GetLocalUniqueNetIdEOS(LocalUserNum))
+	{
+		if (const FUserOnlineAccountEOSRef* Entry = UniqueNetIdToUserAccountMap.Find(NetIdPtr.ToSharedRef()))
 		{
-			OnlineUser = NetIdStringToOnlineUserMap.FindRef(NetId->ToString());
+			OnlineUser = *Entry;
 		}
 	}
+
 	return OnlineUser;
 }
 
 FOnlineUserPtr FUserManagerEOS::GetOnlineUser(EOS_ProductUserId UserId) const
 {
 	FOnlineUserPtr OnlineUser;
-	if (ProductUserIdToStringMap.Contains(UserId))
+
+	if (const FUniqueNetIdEOSPtr NetIdPtr = FUniqueNetIdEOSRegistry::Find(UserId))
 	{
-		const FString& NetId = ProductUserIdToStringMap.FindRef(UserId);
-		if (NetIdStringToOnlineUserMap.Contains(NetId))
+		if (const FUserOnlineAccountEOSRef* Entry = UniqueNetIdToUserAccountMap.Find(NetIdPtr.ToSharedRef()))
 		{
-			OnlineUser = NetIdStringToOnlineUserMap.FindRef(*NetId);
+			OnlineUser = *Entry;
 		}
 	}
+
 	return OnlineUser;
 }
 
 FOnlineUserPtr FUserManagerEOS::GetOnlineUser(EOS_EpicAccountId AccountId) const
 {
 	FOnlineUserPtr OnlineUser;
-	if (AccountIdToStringMap.Contains(AccountId))
+
+	if (const FUniqueNetIdEOSPtr NetIdPtr = FUniqueNetIdEOSRegistry::Find(AccountId))
 	{
-		const FString& NetId = AccountIdToStringMap.FindRef(AccountId);
-		if (NetIdStringToOnlineUserMap.Contains(NetId))
+		if (const FUserOnlineAccountEOSRef* Entry = UniqueNetIdToUserAccountMap.Find(NetIdPtr.ToSharedRef()))
 		{
-			OnlineUser = NetIdStringToOnlineUserMap.FindRef(NetId);
+			OnlineUser = *Entry;
 		}
 	}
+	
 	return OnlineUser;
 }
 
@@ -1434,11 +1433,8 @@ void FUserManagerEOS::RemoveLocalUser(int32 LocalUserNum)
 		EOSSubsystem->ReleaseVoiceChatUserInterface(*FoundId);
 		LocalUserNumToFriendsListMap.Remove(LocalUserNum);
 		const FString& NetId = FoundId->ToString();
-		AccountIdToStringMap.Remove(FoundId->GetEpicAccountId());
-		NetIdStringToOnlineUserMap.Remove(NetId);
 		LocalUserUniqueNetIds.RemoveAt(LocalUserNum);
 		UniqueNetIdToUserAccountMap.Remove(FoundId);
-		ProductUserIdToStringMap.Remove(FoundId->GetProductUserId());
 	}
 	// Reset this for the next user login
 	if (LocalUserNum == DefaultLocalUser)
@@ -1734,41 +1730,41 @@ void FUserManagerEOS::FriendStatusChanged(const EOS_Friends_OnFriendsUpdateInfo*
 	{
 		const int32 LocalUserNum = GetLocalUserNumFromEpicAccountId(Data->LocalUserId);
 		const FUniqueNetIdEOSRef LocalEOSID = LocalUserUniqueNetIds[LocalUserNum];
-		if (!AccountIdToStringMap.Contains(Data->TargetUserId))
+		if (!FUniqueNetIdEOSRegistry::Find(Data->TargetUserId))
 		{
-			FUniqueNetIdEOSRef FriendNetId = FUniqueNetIdEOSRegistry::FindOrAdd(Data->TargetUserId, nullptr).ToSharedRef();
+			FUniqueNetIdEOSRef FriendNetId = FUniqueNetIdEOSRegistry::FindOrAdd(Data->TargetUserId, nullptr);
 			const FString NetId = FriendNetId->ToString();
 			AddRemotePlayer(LocalUserNum, NetId, Data->TargetUserId);
 			// TODO: We need to modify AddRemotePlayer to take a delegate, to execute after all information has been retrieved
 		}
 
-		FOnlineUserPtr OnlineUser = EpicAccountIdToOnlineUserMap[Data->TargetUserId];
+		// AddRemotePlayer just added this id to the registry, so it will be found and valid, no checks needed
+		const FUniqueNetIdEOS& FriendEOSId = *FUniqueNetIdEOSRegistry::FindChecked(Data->TargetUserId);
 
 		switch (Data->CurrentStatus)
 		{
 		case EOS_EFriendsStatus::EOS_FS_NotFriends: // Invite rejections and friend removal
 		{
 			//User should already be a friend
-			const FString& AccountIdStr = AccountIdToStringMap[Data->TargetUserId];
 			FFriendsListEOSRef& FriendsListRef = LocalUserNumToFriendsListMap[LocalUserNum];
-			FOnlineFriendEOSPtr Friend = FriendsListRef->GetByNetIdString(AccountIdStr);
+			FOnlineFriendEOSPtr Friend = FriendsListRef->GetByNetIdString(FriendEOSId.ToString());
 			if (Friend.IsValid())
 			{
-				FriendsListRef->Remove(AccountIdStr, Friend.ToSharedRef());
+				FriendsListRef->Remove(FriendEOSId.ToString(), Friend.ToSharedRef());
 				Friend->SetInviteStatus(EInviteStatus::Unknown);
 
 				if (Data->PreviousStatus == EOS_EFriendsStatus::EOS_FS_Friends)
 				{
-					TriggerOnFriendRemovedDelegates(*LocalEOSID, *OnlineUser->GetUserId());
+					TriggerOnFriendRemovedDelegates(*LocalEOSID, FriendEOSId);
 				}
 				else if (Data->PreviousStatus == EOS_EFriendsStatus::EOS_FS_InviteSent || Data->PreviousStatus == EOS_EFriendsStatus::EOS_FS_InviteReceived)
 				{
-					TriggerOnInviteRejectedDelegates(*LocalEOSID, *OnlineUser->GetUserId()); // We don't have an "OnInviteRejected" event only for the local user
+					TriggerOnInviteRejectedDelegates(*LocalEOSID, FriendEOSId); // We don't have an "OnInviteRejected" event only for the local user
 				}
 			}
 			else
 			{
-				UE_LOG_ONLINE_FRIEND(Verbose, TEXT("Friend status notification received for user [%d], but remote user [%s] was not previously registered as a friend"), *AccountIdStr);
+				UE_LOG_ONLINE_FRIEND(Verbose, TEXT("Friend status notification received for user [%d], but remote user [%s] was not previously registered as a friend"), *FriendEOSId.ToString());
 			}
 
 			break;
@@ -1793,7 +1789,7 @@ void FUserManagerEOS::FriendStatusChanged(const EOS_Friends_OnFriendsUpdateInfo*
 			FOnlineFriendEOSRef FriendRef = AddFriend(LocalUserNum, Data->TargetUserId);
 
 			FriendRef->SetInviteStatus(EInviteStatus::PendingInbound);
-			TriggerOnInviteReceivedDelegates(*LocalEOSID, *OnlineUser->GetUserId());
+			TriggerOnInviteReceivedDelegates(*LocalEOSID, FriendEOSId);
 
 			break;
 		}
@@ -1802,7 +1798,7 @@ void FUserManagerEOS::FriendStatusChanged(const EOS_Friends_OnFriendsUpdateInfo*
 			FOnlineFriendEOSRef FriendRef = AddFriend(LocalUserNum, Data->TargetUserId);
 
 			FriendRef->SetInviteStatus(EInviteStatus::Accepted);
-			TriggerOnInviteAcceptedDelegates(*LocalEOSID, *OnlineUser->GetUserId()); // We don't have an "OnFriendAdded", all relationships exist in the context of invites. We also don't have a separate "OnInviteAccepted" for the local user.
+			TriggerOnInviteAcceptedDelegates(*LocalEOSID, FriendEOSId); // We don't have an "OnFriendAdded", all relationships exist in the context of invites. We also don't have a separate "OnInviteAccepted" for the local user.
 
 			break;
 		}
@@ -1817,7 +1813,7 @@ void FUserManagerEOS::FriendStatusChanged(const EOS_Friends_OnFriendsUpdateInfo*
 
 FOnlineFriendEOSRef FUserManagerEOS::AddFriend(int32 LocalUserNum, EOS_EpicAccountId EpicAccountId)
 {
-	FUniqueNetIdEOSRef FriendNetId = FUniqueNetIdEOSRegistry::FindOrAdd(EpicAccountId, nullptr).ToSharedRef();
+	FUniqueNetIdEOSRef FriendNetId = FUniqueNetIdEOSRegistry::FindOrAdd(EpicAccountId, nullptr);
 	const FString NetId = FriendNetId->ToString();
 	FOnlineFriendEOSRef FriendRef = MakeShareable(new FOnlineFriendEOS(FriendNetId));
 	LocalUserNumToFriendsListMap[LocalUserNum]->Add(NetId, FriendRef);
@@ -1832,7 +1828,7 @@ FOnlineFriendEOSRef FUserManagerEOS::AddFriend(int32 LocalUserNum, EOS_EpicAccou
 	FriendRef->SetInviteStatus(ToEInviteStatus(Status));
 
 	// Add this friend as a remote player if we haven't already (this will grab user info)
-	if (!AccountIdToStringMap.Contains(EpicAccountId))
+	if (!FUniqueNetIdEOSRegistry::Find(EpicAccountId))
 	{
 		AddRemotePlayer(LocalUserNum, NetId, EpicAccountId, FriendNetId, FriendRef, FriendRef);
 	}
@@ -1850,7 +1846,7 @@ FOnlineFriendEOSRef FUserManagerEOS::AddFriend(int32 LocalUserNum, EOS_EpicAccou
 
 void FUserManagerEOS::AddRemotePlayer(int32 LocalUserNum, const FString& NetId, EOS_EpicAccountId EpicAccountId)
 {
-	FUniqueNetIdEOSRef EOSID = FUniqueNetIdEOSRegistry::FindOrAdd(NetId).ToSharedRef();
+	FUniqueNetIdEOSRef EOSID = FUniqueNetIdEOSRegistry::FindOrAdd(NetId);
 	FOnlineUserEOSRef UserRef = MakeShareable(new FOnlineUserEOS(EOSID));
 	// Add this user as a remote (this will grab presence & user info)
 	AddRemotePlayer(LocalUserNum, NetId, EpicAccountId, EOSID, UserRef, UserRef);
@@ -1858,34 +1854,33 @@ void FUserManagerEOS::AddRemotePlayer(int32 LocalUserNum, const FString& NetId,
 
 void FUserManagerEOS::AddRemotePlayer(int32 LocalUserNum, const FString& NetId, EOS_EpicAccountId EpicAccountId, FUniqueNetIdEOSPtr UniqueNetId, FOnlineUserPtr OnlineUser, IAttributeAccessInterfaceRef AttributeRef)
 {
-	NetIdStringToOnlineUserMap.Emplace(NetId, OnlineUser);
-	EpicAccountIdToOnlineUserMap.Emplace(EpicAccountId, OnlineUser);
+	// TODO: pending asynchronous resolution of the remote player's UniqueNetId
+	FUniqueNetIdEOSRegistry::FindOrAdd(EpicAccountId, nullptr);
+	
 	NetIdStringToAttributeAccessMap.Emplace(NetId, AttributeRef);
 	EpicAccountIdToAttributeAccessMap.Emplace(EpicAccountId, AttributeRef);
 
-	AccountIdToStringMap.Emplace(EpicAccountId, NetId);
-
 	// Read the user info for this player
 	ReadUserInfo(LocalUserNum, EpicAccountId);
 }
 
 void FUserManagerEOS::UpdateRemotePlayerProductUserId(EOS_EpicAccountId EpicAccountId, EOS_ProductUserId ProductUserId)
 {
-	// See if the net ids have changed for this user and bail if they are the same
-	const FString PrevNetIdStr = AccountIdToStringMap[EpicAccountId];
+	const FUniqueNetIdEOSRef NewNetId = FUniqueNetIdEOSRegistry::FindOrAdd(EpicAccountId, ProductUserId);
+
+	const FString PrevNetIdStr = LexToString(EpicAccountId) + EOS_ID_SEPARATOR;
 	// Note that calling FindOrAdd with previously invalid EAS/EOS components will update the net id in place.
-	const FString NewNetIdStr = FUniqueNetIdEOSRegistry::FindOrAdd(EpicAccountId, ProductUserId)->ToString();
-	if (PrevNetIdStr == NewNetIdStr)
+	const FString NewNetIdStr = NewNetId->ToString();
+	
+	// TODO: While there is still unresolved ids and string maps in the system, we'll force this method to always update
+	/*if (PrevNetIdStr == NewNetIdStr)
 	{
 		// No change, so skip any work
 		return;
-	}
+	}*/
 
 	const FString ProductUserIdStr = LexToString(ProductUserId);
 
-	// Get the unique net id and rebuild the string for it
-	IAttributeAccessInterfaceRef AttrAccess = NetIdStringToAttributeAccessMap[PrevNetIdStr];
-	FUniqueNetIdEOSPtr NetIdEOS = AttrAccess->GetUniqueNetIdEOS();
 	// Update any old friends entries with the new net id key
 	for (TMap<int32, FFriendsListEOSRef>::TConstIterator It(LocalUserNumToFriendsListMap); It; ++It)
 	{
@@ -1896,16 +1891,15 @@ void FUserManagerEOS::UpdateRemotePlayerProductUserId(EOS_EpicAccountId EpicAcco
 			FriendsList->UpdateNetIdStr(PrevNetIdStr, NewNetIdStr);
 		}
 	}
+
 	// Update all of the other net id to X mappings
-	AccountIdToStringMap.Remove(EpicAccountId);
-	AccountIdToStringMap.Emplace(EpicAccountId, NewNetIdStr);
-	ProductUserIdToStringMap.Remove(ProductUserId);
-	ProductUserIdToStringMap.Emplace(ProductUserId, *NewNetIdStr);
-	FOnlineUserPtr OnlineUser = NetIdStringToOnlineUserMap[PrevNetIdStr];
-	NetIdStringToOnlineUserMap.Remove(PrevNetIdStr);
-	NetIdStringToOnlineUserMap.Emplace(NewNetIdStr, OnlineUser);
-	NetIdStringToAttributeAccessMap.Remove(PrevNetIdStr);
-	NetIdStringToAttributeAccessMap.Emplace(NewNetIdStr, AttrAccess);
+	if (NetIdStringToAttributeAccessMap.Contains(PrevNetIdStr))
+	{
+		const IAttributeAccessInterfaceRef AttrAccess = NetIdStringToAttributeAccessMap[PrevNetIdStr];
+		NetIdStringToAttributeAccessMap.Remove(PrevNetIdStr);
+		NetIdStringToAttributeAccessMap.Emplace(NewNetIdStr, AttrAccess);
+	}
+
 	// Presence may not be available for all online users
 	if (NetIdStringToOnlineUserPresenceMap.Contains(PrevNetIdStr))
 	{
@@ -2087,15 +2081,16 @@ bool FUserManagerEOS::SendInvite(int32 LocalUserNum, const FUniqueNetId& FriendI
 	FSendInviteCallback* CallbackObj = new FSendInviteCallback(AsWeak());
 	CallbackObj->CallbackLambda = [LocalUserNum, ListName, this, Delegate](const EOS_Friends_SendInviteCallbackInfo* Data)
 	{
-		const FString& NetId = AccountIdToStringMap[Data->TargetUserId];
+		const FUniqueNetIdEOSRef NetIdRef = FUniqueNetIdEOSRegistry::FindChecked(Data->TargetUserId);
 
 		FString ErrorString;
 		bool bWasSuccessful = Data->ResultCode == EOS_EResult::EOS_Success;
 		if (!bWasSuccessful)
 		{
-			ErrorString = FString::Printf(TEXT("Failed to send invite for user (%d) to player (%s) with result code (%s)"), LocalUserNum, *NetId, ANSI_TO_TCHAR(EOS_EResult_ToString(Data->ResultCode)));
+			ErrorString = FString::Printf(TEXT("Failed to send invite for user (%d) to player (%s) with result code (%s)"), LocalUserNum, *NetIdRef->ToString(), ANSI_TO_TCHAR(EOS_EResult_ToString(Data->ResultCode)));
 		}
-		Delegate.ExecuteIfBound(LocalUserNum, bWasSuccessful, *FUniqueNetIdEOSRegistry::FindOrAdd(NetId), ListName, ErrorString);
+
+		Delegate.ExecuteIfBound(LocalUserNum, bWasSuccessful, *NetIdRef, ListName, ErrorString);
 	};
 
 	EOS_Friends_SendInviteOptions Options = { };
@@ -2131,15 +2126,15 @@ bool FUserManagerEOS::AcceptInvite(int32 LocalUserNum, const FUniqueNetId& Frien
 	FAcceptInviteCallback* CallbackObj = new FAcceptInviteCallback(AsWeak());
 	CallbackObj->CallbackLambda = [LocalUserNum, ListName, this, Delegate](const EOS_Friends_AcceptInviteCallbackInfo* Data)
 	{
-		const FString& NetId = AccountIdToStringMap[Data->TargetUserId];
+		const FUniqueNetIdEOSRef NetIdRef = FUniqueNetIdEOSRegistry::FindChecked(Data->TargetUserId);
 
 		FString ErrorString;
 		bool bWasSuccessful = Data->ResultCode == EOS_EResult::EOS_Success;
 		if (!bWasSuccessful)
 		{
-			ErrorString = FString::Printf(TEXT("Failed to accept invite for user (%d) from friend (%s) with result code (%s)"), LocalUserNum, *NetId, ANSI_TO_TCHAR(EOS_EResult_ToString(Data->ResultCode)));
+			ErrorString = FString::Printf(TEXT("Failed to accept invite for user (%d) from friend (%s) with result code (%s)"), LocalUserNum, *NetIdRef->ToString(), ANSI_TO_TCHAR(EOS_EResult_ToString(Data->ResultCode)));
 		}
-		Delegate.ExecuteIfBound(LocalUserNum, bWasSuccessful, *FUniqueNetIdEOSRegistry::FindOrAdd(NetId), ListName, ErrorString);
+		Delegate.ExecuteIfBound(LocalUserNum, bWasSuccessful, *NetIdRef, ListName, ErrorString);
 	};
 
 	EOS_Friends_AcceptInviteOptions Options = { };
@@ -2413,7 +2408,7 @@ bool FUserManagerEOS::HandleFriendsExec(UWorld* InWorld, const TCHAR* Cmd, FOutp
 		int LocalUserNum = FCString::Atoi(*FParse::Token(Cmd, false));
 
 		FString FriendUserIdStr = FParse::Token(Cmd, false);
-		FUniqueNetIdEOSRef FriendEosId = FUniqueNetIdEOSRegistry::FindOrAdd(FriendUserIdStr).ToSharedRef();
+		FUniqueNetIdEOSRef FriendEosId = FUniqueNetIdEOSRegistry::FindOrAdd(FriendUserIdStr);
 
 		FString FriendsList = FParse::Token(Cmd, false);
 
@@ -2437,7 +2432,7 @@ bool FUserManagerEOS::HandleFriendsExec(UWorld* InWorld, const TCHAR* Cmd, FOutp
 		int LocalUserNum = FCString::Atoi(*FParse::Token(Cmd, false));
 
 		FString FriendUserIdStr = FParse::Token(Cmd, false);
-		FUniqueNetIdEOSRef FriendEosId = FUniqueNetIdEOSRegistry::FindOrAdd(FriendUserIdStr).ToSharedRef();
+		FUniqueNetIdEOSRef FriendEosId = FUniqueNetIdEOSRegistry::FindOrAdd(FriendUserIdStr);
 
 		FString FriendsList = FParse::Token(Cmd, false);
 
@@ -2552,14 +2547,9 @@ void FUserManagerEOS::SetPresence(const FUniqueNetId& UserId, const FOnlineUserP
 	FSetPresenceCallback* CallbackObj = new FSetPresenceCallback(AsWeak());
 	CallbackObj->CallbackLambda = [this, Delegate](const EOS_Presence_SetPresenceCallbackInfo* Data)
 	{
-		if (Data->ResultCode == EOS_EResult::EOS_Success && AccountIdToStringMap.Contains(Data->LocalUserId))
-		{
-			FUniqueNetIdEOSRef EOSID = FUniqueNetIdEOSRegistry::FindOrAdd(AccountIdToStringMap[Data->LocalUserId]).ToSharedRef();
-			Delegate.ExecuteIfBound(*EOSID, true);
-			return;
-		}
-		UE_LOG_ONLINE(Error, TEXT("SetPresence() failed with result code (%s)"), *LexToString(Data->ResultCode));
-		Delegate.ExecuteIfBound(*FUniqueNetIdEOS::EmptyId(), false);
+		const bool bSuccess = Data->ResultCode == EOS_EResult::EOS_Success;
+		UE_CLOG_ONLINE(!bSuccess, Error, TEXT("SetPresence() failed with result code (%s)"), *LexToString(Data->ResultCode));
+		Delegate.ExecuteIfBound(*GetLocalUniqueNetIdEOS(Data->LocalUserId), bSuccess);
 	};
 
 	EOS_Presence_SetPresenceOptions PresOptions = { };
@@ -2603,14 +2593,17 @@ void FUserManagerEOS::QueryPresence(const FUniqueNetId& UserId, const FOnPresenc
 		FQueryPresenceCallback* CallbackObj = new FQueryPresenceCallback(AsWeak());
 		CallbackObj->CallbackLambda = [this, Delegate](const EOS_Presence_QueryPresenceCallbackInfo* Data)
 		{
-			if (Data->ResultCode == EOS_EResult::EOS_Success && EpicAccountIdToOnlineUserMap.Contains(Data->TargetUserId))
+			if (Data->ResultCode == EOS_EResult::EOS_Success)
 			{
-				// Update the presence data to the most recent
-				UpdatePresence(DefaultLocalUser, Data->TargetUserId);
-				FOnlineUserPtr OnlineUser = EpicAccountIdToOnlineUserMap[Data->TargetUserId];
-				Delegate.ExecuteIfBound(*OnlineUser->GetUserId(), true);
-				return;
+				if (const FUniqueNetIdEOSPtr NetIdPtr = FUniqueNetIdEOSRegistry::Find(Data->TargetUserId))
+				{
+					// Update the presence data to the most recent
+					UpdatePresence(DefaultLocalUser, Data->TargetUserId);
+					Delegate.ExecuteIfBound(*NetIdPtr, true);
+					return;
+				}
 			}
+
 			const FString& TargetUser = FUniqueNetIdEOSRegistry::FindOrAdd(Data->TargetUserId, nullptr)->ToString();
 			UE_LOG_ONLINE(Error, TEXT("QueryPresence() for user (%s) failed with result code (%s)"), *TargetUser, *LexToString(Data->ResultCode));
 			Delegate.ExecuteIfBound(*FUniqueNetIdEOS::EmptyId(), false);
@@ -2643,15 +2636,16 @@ void FUserManagerEOS::UpdatePresence(int32 LocalUserNum, EOS_EpicAccountId Accou
 	EOS_EResult CopyResult = EOS_Presence_CopyPresence(EOSSubsystem->PresenceHandle, &Options, &PresenceInfo);
 	if (CopyResult == EOS_EResult::EOS_Success)
 	{
-		const FString& NetId = AccountIdToStringMap[AccountId];
+		const FString& NetIdStr = FUniqueNetIdEOSRegistry::FindChecked(AccountId)->ToString();
+
 		// Create it on demand if we don't have one yet
-		if (!NetIdStringToOnlineUserPresenceMap.Contains(NetId))
+		if (!NetIdStringToOnlineUserPresenceMap.Contains(NetIdStr))
 		{
 			FOnlineUserPresenceRef PresenceRef = MakeShareable(new FOnlineUserPresence());
-			NetIdStringToOnlineUserPresenceMap.Emplace(NetId, PresenceRef);
+			NetIdStringToOnlineUserPresenceMap.Emplace(NetIdStr, PresenceRef);
 		}
 
-		FOnlineUserPresenceRef PresenceRef = NetIdStringToOnlineUserPresenceMap[NetId];
+		FOnlineUserPresenceRef PresenceRef = NetIdStringToOnlineUserPresenceMap[NetIdStr];
 		const FString ProductId(UTF8_TO_TCHAR(PresenceInfo->ProductId));
 		const FString ProdVersion(UTF8_TO_TCHAR(PresenceInfo->ProductVersion));
 		const FString Platform(UTF8_TO_TCHAR(PresenceInfo->Platform));
@@ -2673,7 +2667,7 @@ void FUserManagerEOS::UpdatePresence(int32 LocalUserNum, EOS_EpicAccountId Accou
 		}
 
 		// Copy the presence if this is a friend that was updated, so that their data is in sync
-		UpdateFriendPresence(NetId, PresenceRef);
+		UpdateFriendPresence(NetIdStr, PresenceRef);
 
 		EOS_Presence_Info_Release(PresenceInfo);
 	}
@@ -2803,30 +2797,28 @@ void FUserManagerEOS::ReadUserInfo(int32 LocalUserNum, EOS_EpicAccountId EpicAcc
 bool FUserManagerEOS::GetAllUserInfo(int32 LocalUserNum, TArray<TSharedRef<FOnlineUser>>& OutUsers)
 {
 	OutUsers.Reset();
-	// Get remote users
-	for (TMap<FString, FOnlineUserPtr>::TConstIterator It(NetIdStringToOnlineUserMap); It; ++It)
+
+	// Get local users only
+	for (const FUniqueNetIdEOSRef& LocalUserNetId : LocalUserUniqueNetIds)
 	{
-		if (It.Value().IsValid())
+		if (const FUserOnlineAccountEOSRef* LocalUser = UniqueNetIdToUserAccountMap.Find(LocalUserNetId))
 		{
-			OutUsers.Add(It.Value().ToSharedRef());
+			OutUsers.Add(*LocalUser);
 		}
 	}
-	// Get local users
-	for (TUniqueNetIdMap<FUserOnlineAccountEOSRef>::TConstIterator It(UniqueNetIdToUserAccountMap); It; ++It)
-	{
-		OutUsers.Add(It.Value());
-	}
+
 	return true;
 }
 
 TSharedPtr<FOnlineUser> FUserManagerEOS::GetUserInfo(int32 LocalUserNum, const FUniqueNetId& UserId)
 {
 	TSharedPtr<FOnlineUser> OnlineUser;
-	const FUniqueNetIdEOS& EOSID = FUniqueNetIdEOS::Cast(UserId);
-	if (NetIdStringToOnlineUserMap.Contains(EOSID.ToString()))
+
+	if (const FUserOnlineAccountEOSRef* FoundUser = UniqueNetIdToUserAccountMap.Find(UserId.AsShared()))
 	{
-		OnlineUser = NetIdStringToOnlineUserMap[EOSID.ToString()];
+		OnlineUser = *FoundUser;
 	}
+
 	return OnlineUser;
 }
 
@@ -2873,7 +2865,7 @@ bool FUserManagerEOS::QueryUserIdMapping(const FUniqueNetId& UserId, const FStri
 		{
 			const FString NetIdStr = LexToString(Data->TargetUserId);
 			const FUniqueNetIdEOSRef LocalUserId = LocalUserUniqueNetIds[LocalUserNum];
-			if (!EpicAccountIdToOnlineUserMap.Contains(Data->TargetUserId))
+			if (!FUniqueNetIdEOSRegistry::Find(Data->TargetUserId))
 			{
 				// Registering the player will also query the presence/user info data
 				AddRemotePlayer(LocalUserNum, NetIdStr, Data->TargetUserId);
@@ -3042,10 +3034,9 @@ FUniqueNetIdPtr FUserManagerEOS::GetExternalIdMapping(const FExternalIdQueryOpti
 {
 	FUniqueNetIdPtr NetId;
 	EOS_EpicAccountId AccountId = EOS_EpicAccountId_FromString(TCHAR_TO_UTF8(*ExternalId));
-	if (EOS_EpicAccountId_IsValid(AccountId) == EOS_TRUE && AccountIdToStringMap.Contains(AccountId))
+	if (EOS_EpicAccountId_IsValid(AccountId) == EOS_TRUE)
 	{
-		const FString& NetIdStr = AccountIdToStringMap[AccountId];
-		NetId = NetIdStringToOnlineUserMap[NetIdStr]->GetUserId();
+		return FUniqueNetIdEOSRegistry::Find(AccountId);
 	}
 	return NetId;
 }
diff --git a/Engine/Plugins/Online/OnlineSubsystemEOS/Source/OnlineSubsystemEOS/Private/UserManagerEOS.h b/Engine/Plugins/Online/OnlineSubsystemEOS/Source/OnlineSubsystemEOS/Private/UserManagerEOS.h
--- a/Engine/Plugins/Online/OnlineSubsystemEOS/Source/OnlineSubsystemEOS/Private/UserManagerEOS.h
+++ b/Engine/Plugins/Online/OnlineSubsystemEOS/Source/OnlineSubsystemEOS/Private/UserManagerEOS.h
@@ -118,9 +118,11 @@ public:
 
 	void UpdateNetIdStr(const FString& PrevNetId, const FString& NewNetId)
 	{
-		ListClass ListEntry = NetIdStringToListEntryMap[PrevNetId];
-		NetIdStringToListEntryMap.Remove(PrevNetId);
-		NetIdStringToListEntryMap.Add(NewNetId, ListEntry);
+		if (ListClass* ListEntry = NetIdStringToListEntryMap.Find(PrevNetId))
+		{
+			NetIdStringToListEntryMap.Add(NewNetId, *ListEntry);
+			NetIdStringToListEntryMap.Remove(PrevNetId);
+		}
 	}
 
 	ListClassReturnType GetByIndex(int32 Index)
@@ -365,17 +367,11 @@ private:
 	/** General map for all registered online users */
 	TUniqueNetIdMap<FUserOnlineAccountEOSRef> UniqueNetIdToUserAccountMap;
 
-	/** General account mappings */
-	TMap<EOS_EpicAccountId, FString> AccountIdToStringMap;
-	TMap<EOS_ProductUserId, FString> ProductUserIdToStringMap;
-
 	/** Per user friends lists accessible by user num or net id */
 	TMap<int32, FFriendsListEOSRef> LocalUserNumToFriendsListMap;
 	TMap<FString, FFriendsListEOSRef> NetIdStringToFriendsListMap;
 
 	/** Ids mapped to remote users */
-	TMap<FString, FOnlineUserPtr> NetIdStringToOnlineUserMap;
-	TMap<EOS_EpicAccountId, FOnlineUserPtr> EpicAccountIdToOnlineUserMap;
 	TMap<FString, IAttributeAccessInterfaceRef> NetIdStringToAttributeAccessMap;
 	TMap<EOS_EpicAccountId, IAttributeAccessInterfaceRef> EpicAccountIdToAttributeAccessMap;
 
