Fix for ResolveNetIds so that Delegate execution remains asynchronous in all code paths.

diff --git a/Engine/Plugins/Online/OnlineSubsystemEOS/Source/OnlineSubsystemEOS/Private/UserManagerEOS.cpp b/Engine/Plugins/Online/OnlineSubsystemEOS/Source/OnlineSubsystemEOS/Private/UserManagerEOS.cpp
--- a/Engine/Plugins/Online/OnlineSubsystemEOS/Source/OnlineSubsystemEOS/Private/UserManagerEOS.cpp
+++ b/Engine/Plugins/Online/OnlineSubsystemEOS/Source/OnlineSubsystemEOS/Private/UserManagerEOS.cpp
@@ -1375,7 +1375,10 @@ void FUserManagerEOS::ResolveUniqueNetIds(int32 LocalUserNum, const TArray<EOS_P
 	}
 	else
 	{
-		Callback(ResolvedUniqueNetIds);
+		EOSSubsystem->ExecuteNextTick([Callback, ResolvedUniqueNetIds]()
+			{
+				Callback(ResolvedUniqueNetIds);
+			});
 	}
 }
 
