FUserManagerEOS internal data structures refactor (Part 3):
- Created FLocalUserEOS type to store all data concerning local users, refactored usage
- Refactored flow to transmit login credentials until local user creation (after a successful login). These credentials are needed later for connect auth refresh.
- Refactored methods to remove unnecessary parameters.
- Refactored Attribute Access and Presence maps to use UniqueIds as keys

diff --git a/Engine/Plugins/Online/OnlineSubsystemEOS/Source/OnlineSubsystemEOS/Private/UserManagerEOS.cpp b/Engine/Plugins/Online/OnlineSubsystemEOS/Source/OnlineSubsystemEOS/Private/UserManagerEOS.cpp
--- a/Engine/Plugins/Online/OnlineSubsystemEOS/Source/OnlineSubsystemEOS/Private/UserManagerEOS.cpp
+++ b/Engine/Plugins/Online/OnlineSubsystemEOS/Source/OnlineSubsystemEOS/Private/UserManagerEOS.cpp
@@ -277,13 +277,22 @@ void FUserManagerEOS::LoginStatusChanged(const EOS_Auth_LoginStatusChangedCallba
 		if (IsLocalUser(Data->LocalUserId))
 		{
 			const int32 LocalUserNum = GetLocalUserNumFromEpicAccountId(Data->LocalUserId);
-			const FUniqueNetIdEOSPtr UserNetId = LocalUserUniqueNetIds[LocalUserNum];
+			FLocalUserEOS& LocalUser = GetLocalUserChecked(LocalUserNum);
+			const FUniqueNetIdEOSPtr& UserNetId = LocalUser.UniqueNetId;
 			TriggerOnLoginStatusChangedDelegates(LocalUserNum, ELoginStatus::LoggedIn, ELoginStatus::NotLoggedIn, *UserNetId);
+
+			// Remove the per user connect login notification
+			if (LocalUser.ConnectLoginNotification.IsValid())
+			{
+				EOS_Connect_RemoveNotifyAuthExpiration(EOSSubsystem->ConnectHandle, LocalUser.ConnectLoginNotification->NotificationId);
+				LocalUser.ConnectLoginNotification.Reset();
+			}
+
 			// Need to remove the local user
 			RemoveLocalUser(LocalUserNum);
 
-			// Clean up user based notifies if we have no logged in users
-			if (LocalUserUniqueNetIds.Num() == 0)
+			// Clean up user based notifications if we have no logged in users
+			if (LocalUsers.Num() == 0)
 			{
 				if (LoginNotificationId > 0)
 				{
@@ -307,16 +316,6 @@ void FUserManagerEOS::LoginStatusChanged(const EOS_Auth_LoginStatusChangedCallba
 					PresenceNotificationCallback = nullptr;
 					PresenceNotificationId = 0;
 				}
-				// Remove the per user connect login notification
-				if (LocalUserNumToConnectLoginNotifcationMap.Contains(LocalUserNum))
-				{
-					FNotificationIdCallbackPair* NotificationPair = LocalUserNumToConnectLoginNotifcationMap[LocalUserNum];
-					LocalUserNumToConnectLoginNotifcationMap.Remove(LocalUserNum);
-
-					EOS_Connect_RemoveNotifyAuthExpiration(EOSSubsystem->ConnectHandle, NotificationPair->NotificationId);
-
-					delete NotificationPair;
-				}
 			}
 		}
 	}
@@ -378,8 +377,6 @@ typedef TEOSCallback<EOS_Auth_OnDeletePersistentAuthCallback, EOS_Auth_DeletePer
 
 bool FUserManagerEOS::Login(int32 LocalUserNum, const FOnlineAccountCredentials& AccountCredentials)
 {
-	LocalUserNumToLastLoginCredentials.Emplace(LocalUserNum, MakeShared<FOnlineAccountCredentials>(AccountCredentials));
-
 	FEOSSettings Settings = UEOSSettings::GetSettings();
 
 	// Are we configured to run at all?
@@ -390,13 +387,6 @@ bool FUserManagerEOS::Login(int32 LocalUserNum, const FOnlineAccountCredentials&
 		return true;
 	}
 
-	// See if we are configured to just use EOS and not EAS
-	if (!EOSSubsystem->bIsDefaultOSS && !EOSSubsystem->bIsPlatformOSS && !Settings.bUseEAS && Settings.bUseEOSConnect)
-	{
-		// Call the EOS + Platform login path
-		return ConnectLoginNoEAS(LocalUserNum);
-	}
-
 	// We don't support offline logged in, so they are either logged in or not
 	if (GetLoginStatus(LocalUserNum) == ELoginStatus::LoggedIn)
 	{
@@ -405,10 +395,20 @@ bool FUserManagerEOS::Login(int32 LocalUserNum, const FOnlineAccountCredentials&
 		return true;
 	}
 
+	// Can't be const because LoginViaExternalAuth might need to set the type
+	FOnlineAccountCredentials CurrentLoginCredentials(AccountCredentials);
+
+	// See if we are configured to just use EOS and not EAS
+	if (!EOSSubsystem->bIsDefaultOSS && !EOSSubsystem->bIsPlatformOSS && !Settings.bUseEAS && Settings.bUseEOSConnect)
+	{
+		// Call the EOS + Platform login path
+		return ConnectLoginNoEAS(LocalUserNum, CurrentLoginCredentials);
+	}
+
 	// See if we are logging in using platform credentials to link to EAS
 	if (!EOSSubsystem->bIsDefaultOSS && !EOSSubsystem->bIsPlatformOSS && Settings.bUseEAS)
 	{
-		LoginViaExternalAuth(LocalUserNum);
+		LoginViaExternalAuth(LocalUserNum, CurrentLoginCredentials);
 		return true;
 	}
 
@@ -429,10 +429,10 @@ FString ToHexString(const TArray<uint8>& InToken)
 	return FString(TokenAnsi);
 }
 
-void FUserManagerEOS::LoginViaExternalAuth(int32 LocalUserNum)
+void FUserManagerEOS::LoginViaExternalAuth(int32 LocalUserNum, const FOnlineAccountCredentials& AccountCredentials)
 {
 	GetPlatformAuthToken(LocalUserNum,
-		FOnGetLinkedAccountAuthTokenCompleteDelegate::CreateLambda([this, WeakThis = AsWeak()](int32 LocalUserNum, bool bWasSuccessful, const FExternalAuthToken& AuthToken)
+		FOnGetLinkedAccountAuthTokenCompleteDelegate::CreateLambda([this, WeakThis = AsWeak(), AccountCredentials = FOnlineAccountCredentials(AccountCredentials)](int32 LocalUserNum, bool bWasSuccessful, const FExternalAuthToken& AuthToken) mutable
 		{
 			if (FUserManagerEOSPtr StrongThis = WeakThis.Pin())
 			{
@@ -443,17 +443,15 @@ void FUserManagerEOS::LoginViaExternalAuth(int32 LocalUserNum)
 					return;
 				}
 
-				check(LocalUserNumToLastLoginCredentials.Contains(LocalUserNum));
-				TSharedRef<FOnlineAccountCredentials> OnlineAccountCredentials = LocalUserNumToLastLoginCredentials[LocalUserNum];
-				OnlineAccountCredentials->Type = FString(TEXT("externalauth"));
+				AccountCredentials.Type = FString(TEXT("externalauth"));
 
 				if (AuthToken.HasTokenData())
 				{
-					OnlineAccountCredentials->Token = ToHexString(AuthToken.TokenData);
+					AccountCredentials.Token = ToHexString(AuthToken.TokenData);
 				}
 				else if (AuthToken.HasTokenString())
 				{
-					OnlineAccountCredentials->Token = AuthToken.TokenString;
+					AccountCredentials.Token = AuthToken.TokenString;
 				}
 				else
 				{
@@ -462,11 +460,17 @@ void FUserManagerEOS::LoginViaExternalAuth(int32 LocalUserNum)
 					return;
 				}
 
-				CallEOSAuthLogin(LocalUserNum, *OnlineAccountCredentials);
+				CallEOSAuthLogin(LocalUserNum, AccountCredentials);
 			}
 		}));
 }
 
+FLocalUserEOS& FUserManagerEOS::GetLocalUserChecked(int32 LocalUserNum)
+{
+	check(LocalUsers.IsValidIndex(LocalUserNum));
+	return LocalUsers[LocalUserNum];
+}
+
 void FUserManagerEOS::CallEOSAuthLogin(int32 LocalUserNum, const FOnlineAccountCredentials& Credentials)
 {
 	// First we construct the EOS Credentials object
@@ -507,17 +511,17 @@ void FUserManagerEOS::CallEOSAuthLogin(int32 LocalUserNum, const FOnlineAccountC
 	const bool bIsPersistentLogin = EOSCredentials.Type == EOS_ELoginCredentialType::EOS_LCT_PersistentAuth;
 
 	FLoginCallback* CallbackObj = new FLoginCallback(AsWeak());
-	CallbackObj->CallbackLambda = [this, LocalUserNum, bIsPersistentLogin](const EOS_Auth_LoginCallbackInfo* Data)
+	CallbackObj->CallbackLambda = [this, LocalUserNum, Credentials, bIsPersistentLogin](const EOS_Auth_LoginCallbackInfo* Data)
 	{
 		if (Data->ResultCode == EOS_EResult::EOS_Success)
 		{
 			// Continue the login process by getting the product user id for EAS only
-			ConnectLoginEAS(LocalUserNum, Data->LocalUserId);
+			ConnectLoginEAS(LocalUserNum, Data->LocalUserId, Credentials);
 		}
 		else if (Data->ResultCode == EOS_EResult::EOS_InvalidUser)
 		{
 			// Link the account
-			LinkEAS(LocalUserNum, Data->ContinuanceToken);
+			LinkEAS(LocalUserNum, Data->ContinuanceToken, Credentials);
 		}
 		else
 		{
@@ -574,16 +578,16 @@ struct FLinkAccountOptions :
 
 typedef TEOSCallback<EOS_Auth_OnLinkAccountCallback, EOS_Auth_LinkAccountCallbackInfo, FUserManagerEOS> FLinkAccountCallback;
 
-void FUserManagerEOS::LinkEAS(int32 LocalUserNum, EOS_ContinuanceToken Token)
+void FUserManagerEOS::LinkEAS(int32 LocalUserNum, EOS_ContinuanceToken Token, const FOnlineAccountCredentials& AccountCredentials)
 {
 	FLinkAccountOptions Options(Token);
 	FLinkAccountCallback* CallbackObj = new FLinkAccountCallback(AsWeak());
-	CallbackObj->CallbackLambda = [this, LocalUserNum](const EOS_Auth_LinkAccountCallbackInfo* Data)
+	CallbackObj->CallbackLambda = [this, LocalUserNum, AccountCredentials](const EOS_Auth_LinkAccountCallbackInfo* Data)
 	{
 		if (Data->ResultCode == EOS_EResult::EOS_Success)
 		{
 			// Continue the login process by getting the product user id
-			ConnectLoginEAS(LocalUserNum, Data->LocalUserId);
+			ConnectLoginEAS(LocalUserNum, Data->LocalUserId, AccountCredentials);
 		}
 		else
 		{
@@ -639,10 +643,10 @@ struct FConnectCredentials :
 	char TokenAnsi[EOS_MAX_TOKEN_SIZE];
 };
 
-bool FUserManagerEOS::ConnectLoginNoEAS(int32 LocalUserNum)
+bool FUserManagerEOS::ConnectLoginNoEAS(int32 LocalUserNum, const FOnlineAccountCredentials& AccountCredentials)
 {
 	GetPlatformAuthToken(LocalUserNum,
-		FOnGetLinkedAccountAuthTokenCompleteDelegate::CreateLambda([this, WeakThis = AsWeak()](int32 LocalUserNum, bool bWasSuccessful, const FExternalAuthToken& AuthToken)
+		FOnGetLinkedAccountAuthTokenCompleteDelegate::CreateLambda([this, WeakThis = AsWeak(), AccountCredentials](int32 LocalUserNum, bool bWasSuccessful, const FExternalAuthToken& AuthToken)
 		{
 			if (FUserManagerEOSPtr StrongThis = WeakThis.Pin())
 			{
@@ -655,8 +659,7 @@ bool FUserManagerEOS::ConnectLoginNoEAS(int32 LocalUserNum)
 				}
 
 				// Now login into our EOS account
-				check(LocalUserNumToLastLoginCredentials.Contains(LocalUserNum));
-				FConnectCredentials Credentials(ToEOS_EExternalCredentialType(GetPlatformOSS()->GetSubsystemName(), *LocalUserNumToLastLoginCredentials[LocalUserNum]), AuthToken);
+				FConnectCredentials Credentials(ToEOS_EExternalCredentialType(GetPlatformOSS()->GetSubsystemName(), AccountCredentials), AuthToken);
 				EOS_Connect_LoginOptions Options = { };
 				Options.ApiVersion = 2;
 				UE_EOS_CHECK_API_MISMATCH(EOS_CONNECT_LOGIN_API_LATEST, 2);
@@ -673,17 +676,17 @@ bool FUserManagerEOS::ConnectLoginNoEAS(int32 LocalUserNum)
 #endif
 
 				FConnectLoginCallback* CallbackObj = new FConnectLoginCallback(AsWeak());
-				CallbackObj->CallbackLambda = [this, LocalUserNum](const EOS_Connect_LoginCallbackInfo* Data)
+				CallbackObj->CallbackLambda = [this, LocalUserNum, AccountCredentials](const EOS_Connect_LoginCallbackInfo* Data)
 				{
 					if (Data->ResultCode == EOS_EResult::EOS_Success)
 					{
 						// We have an account mapping to the platform account, skip to final login
-						FullLoginCallback(LocalUserNum, nullptr, Data->LocalUserId);
+						FullLoginCallback(LocalUserNum, nullptr, Data->LocalUserId, AccountCredentials);
 					}
 					else if (Data->ResultCode == EOS_EResult::EOS_InvalidUser)
 					{
 						// We need to create the platform account mapping for this user using the continuation token
-						CreateConnectedLogin(LocalUserNum, nullptr, Data->ContinuanceToken);
+						CreateConnectedLogin(LocalUserNum, nullptr, Data->ContinuanceToken, AccountCredentials);
 					}
 					else
 					{
@@ -699,7 +702,7 @@ bool FUserManagerEOS::ConnectLoginNoEAS(int32 LocalUserNum)
 	return true;
 }
 
-bool FUserManagerEOS::ConnectLoginEAS(int32 LocalUserNum, EOS_EpicAccountId AccountId)
+bool FUserManagerEOS::ConnectLoginEAS(int32 LocalUserNum, EOS_EpicAccountId AccountId, const FOnlineAccountCredentials& AccountCredentials)
 {
 	EOS_Auth_Token* AuthToken = nullptr;
 	EOS_Auth_CopyUserAuthTokenOptions CopyOptions = { };
@@ -721,17 +724,17 @@ bool FUserManagerEOS::ConnectLoginEAS(int32 LocalUserNum, EOS_EpicAccountId Acco
 		Options.Credentials = &Credentials;
 
 		FConnectLoginCallback* CallbackObj = new FConnectLoginCallback(AsWeak());
-		CallbackObj->CallbackLambda = [LocalUserNum, AccountId, this](const EOS_Connect_LoginCallbackInfo* Data)
+		CallbackObj->CallbackLambda = [LocalUserNum, AccountId, AccountCredentials, this](const EOS_Connect_LoginCallbackInfo* Data)
 		{
 			if (Data->ResultCode == EOS_EResult::EOS_Success)
 			{
 				// We have an account mapping, skip to final login
-				FullLoginCallback(LocalUserNum, AccountId, Data->LocalUserId);
+				FullLoginCallback(LocalUserNum, AccountId, Data->LocalUserId, AccountCredentials);
 			}
 			else if (Data->ResultCode == EOS_EResult::EOS_InvalidUser)
 			{
 				// We need to create the mapping for this user using the continuation token
-				CreateConnectedLogin(LocalUserNum, AccountId, Data->ContinuanceToken);
+				CreateConnectedLogin(LocalUserNum, AccountId, Data->ContinuanceToken, AccountCredentials);
 			}
 			else
 			{
@@ -772,8 +775,7 @@ void FUserManagerEOS::RefreshConnectLogin(int32 LocalUserNum)
 		if (CopyResult == EOS_EResult::EOS_Success)
 		{
 			// We update the auth token cached in the user account, along with the user information
-			const FUniqueNetIdEOSRef UniqueNetId = LocalUserUniqueNetIds[LocalUserNum];
-			const FUserOnlineAccountEOSRef UserAccountRef = UniqueNetIdToUserAccountMap.FindChecked(UniqueNetId);
+			const FUserOnlineAccountEOSRef UserAccountRef = GetLocalUserChecked(LocalUserNum).UserOnlineAccount.ToSharedRef();
 			UserAccountRef->SetAuthAttribute(AUTH_ATTR_ID_TOKEN, AuthToken->AccessToken);
 			UpdateUserInfo(UserAccountRef, AccountId, AccountId);
 
@@ -823,8 +825,7 @@ void FUserManagerEOS::RefreshConnectLogin(int32 LocalUserNum)
 					}
 
 					// Now login into our EOS account
-					check(LocalUserNumToLastLoginCredentials.Contains(LocalUserNum));
-					const FOnlineAccountCredentials& Creds = *LocalUserNumToLastLoginCredentials[LocalUserNum];
+					const FOnlineAccountCredentials& Creds = *GetLocalUserChecked(LocalUserNum).LastLoginCredentials;
 					EOS_EExternalCredentialType CredType = ToEOS_EExternalCredentialType(GetPlatformOSS()->GetSubsystemName(), Creds);
 					FConnectCredentials Credentials(CredType, AuthToken);
 					EOS_Connect_LoginOptions Options = { };
@@ -849,7 +850,7 @@ void FUserManagerEOS::RefreshConnectLogin(int32 LocalUserNum)
 
 typedef TEOSCallback<EOS_Connect_OnCreateUserCallback, EOS_Connect_CreateUserCallbackInfo, FUserManagerEOS> FCreateUserCallback;
 
-void FUserManagerEOS::CreateConnectedLogin(int32 LocalUserNum, EOS_EpicAccountId AccountId, EOS_ContinuanceToken Token)
+void FUserManagerEOS::CreateConnectedLogin(int32 LocalUserNum, EOS_EpicAccountId AccountId, EOS_ContinuanceToken Token, const FOnlineAccountCredentials& AccountCredentials)
 {
 	EOS_Connect_CreateUserOptions Options = { };
 	Options.ApiVersion = 1;
@@ -857,11 +858,11 @@ void FUserManagerEOS::CreateConnectedLogin(int32 LocalUserNum, EOS_EpicAccountId
 	Options.ContinuanceToken = Token;
 
 	FCreateUserCallback* CallbackObj = new FCreateUserCallback(AsWeak());
-	CallbackObj->CallbackLambda = [LocalUserNum, AccountId, this](const EOS_Connect_CreateUserCallbackInfo* Data)
+	CallbackObj->CallbackLambda = [LocalUserNum, AccountId, AccountCredentials, this](const EOS_Connect_CreateUserCallbackInfo* Data)
 	{
 		if (Data->ResultCode == EOS_EResult::EOS_Success)
 		{
-			FullLoginCallback(LocalUserNum, AccountId, Data->LocalUserId);
+			FullLoginCallback(LocalUserNum, AccountId, Data->LocalUserId, AccountCredentials);
 		}
 		else
 		{
@@ -878,7 +879,7 @@ typedef TEOSGlobalCallback<EOS_Presence_OnPresenceChangedCallback, EOS_Presence_
 typedef TEOSGlobalCallback<EOS_Friends_OnFriendsUpdateCallback, EOS_Friends_OnFriendsUpdateInfo, FUserManagerEOS> FFriendsStatusUpdateCallback;
 typedef TEOSGlobalCallback<EOS_Auth_OnLoginStatusChangedCallback, EOS_Auth_LoginStatusChangedCallbackInfo, FUserManagerEOS> FLoginStatusChangedCallback;
 
-void FUserManagerEOS::FullLoginCallback(int32 LocalUserNum, EOS_EpicAccountId AccountId, EOS_ProductUserId UserId)
+void FUserManagerEOS::FullLoginCallback(int32 LocalUserNum, EOS_EpicAccountId AccountId, EOS_ProductUserId UserId, const FOnlineAccountCredentials& AccountCredentials)
 {
 	// Add our login status changed callback if not already set
 	if (LoginNotificationId == 0)
@@ -930,31 +931,12 @@ void FUserManagerEOS::FullLoginCallback(int32 LocalUserNum, EOS_EpicAccountId Ac
 		UE_EOS_CHECK_API_MISMATCH(EOS_PRESENCE_ADDNOTIFYONPRESENCECHANGED_API_LATEST, 1);
 		PresenceNotificationId = EOS_Presence_AddNotifyOnPresenceChanged(EOSSubsystem->PresenceHandle, &Options, CallbackObj, CallbackObj->GetCallbackPtr());
 	}
-	// Add auth refresh notification if not set for this user yet
-	if (!LocalUserNumToConnectLoginNotifcationMap.Contains(LocalUserNum))
-	{
-		FNotificationIdCallbackPair* NotificationPair = new FNotificationIdCallbackPair();
-		LocalUserNumToConnectLoginNotifcationMap.Emplace(LocalUserNum, NotificationPair);
-
-		FRefreshAuthCallback* CallbackObj = new FRefreshAuthCallback(AsWeak());
-		NotificationPair->Callback = CallbackObj;
-		CallbackObj->CallbackLambda = [LocalUserNum, this](const EOS_Connect_AuthExpirationCallbackInfo* Data)
-		{
-			RefreshConnectLogin(LocalUserNum);
-		};
-
-		EOS_Connect_AddNotifyAuthExpirationOptions Options = { };
-		Options.ApiVersion = 1;
-		UE_EOS_CHECK_API_MISMATCH(EOS_CONNECT_ADDNOTIFYAUTHEXPIRATION_API_LATEST, 1);
-		NotificationPair->NotificationId = EOS_Connect_AddNotifyAuthExpiration(EOSSubsystem->ConnectHandle, &Options, CallbackObj, CallbackObj->GetCallbackPtr());
-	}
 
-	AddLocalUser(LocalUserNum, AccountId, UserId);
-	FUniqueNetIdEOSPtr UserNetId = GetLocalUniqueNetIdEOS(LocalUserNum);
-	check(UserNetId.IsValid());
+	// We add the local user
+	const FLocalUserEOS& LocalUser = AddLocalUser(LocalUserNum, AccountId, UserId, AccountCredentials);
 
-	TriggerOnLoginCompleteDelegates(LocalUserNum, true, *UserNetId, FString());
-	TriggerOnLoginStatusChangedDelegates(LocalUserNum, ELoginStatus::NotLoggedIn, ELoginStatus::LoggedIn, *UserNetId);
+	TriggerOnLoginCompleteDelegates(LocalUserNum, true, *LocalUser.UniqueNetId, FString());
+	TriggerOnLoginStatusChangedDelegates(LocalUserNum, ELoginStatus::NotLoggedIn, ELoginStatus::LoggedIn, *LocalUser.UniqueNetId);
 }
 
 typedef TEOSCallback<EOS_Auth_OnLogoutCallback, EOS_Auth_LogoutCallbackInfo, FUserManagerEOS> FLogoutCallback;
@@ -1002,8 +984,6 @@ bool FUserManagerEOS::Logout(int32 LocalUserNum)
 
 	EOS_Auth_Logout(EOSSubsystem->AuthHandle, &LogoutOptions, CallbackObj, CallbackObj->GetCallbackPtr());
 
-	LocalUserNumToLastLoginCredentials.Remove(LocalUserNum);
-
 	return true;
 }
 
@@ -1026,12 +1006,10 @@ bool FUserManagerEOS::AutoLogin(int32 LocalUserNum)
 	}
 	FOnlineAccountCredentials Creds(AuthType, LoginId, Password);
 
-	LocalUserNumToLastLoginCredentials.Emplace(LocalUserNum, MakeShared<FOnlineAccountCredentials>(Creds));
-
 	return Login(LocalUserNum, Creds);
 }
 
-void FUserManagerEOS::AddLocalUser(int32 LocalUserNum, EOS_EpicAccountId EpicAccountId, EOS_ProductUserId UserId)
+FLocalUserEOS& FUserManagerEOS::AddLocalUser(int32 LocalUserNum, EOS_EpicAccountId EpicAccountId, EOS_ProductUserId UserId, const FOnlineAccountCredentials& AccountCredentials)
 {
 	// Set the default user to the first one that logs in
 	if (DefaultLocalUser == INVALID_LOCAL_USER)
@@ -1039,18 +1017,45 @@ void FUserManagerEOS::AddLocalUser(int32 LocalUserNum, EOS_EpicAccountId EpicAcc
 		DefaultLocalUser = LocalUserNum;
 	}
 
+	// Init player lists
+	if (!LocalUsers.IsValidIndex(LocalUserNum))
+	{
+		LocalUsers.Insert(LocalUserNum, FLocalUserEOS());
+	}
+	FLocalUserEOS& LocalUser = GetLocalUserChecked(LocalUserNum);
+
 	FUniqueNetIdEOSRef UserNetId = FUniqueNetIdEOSRegistry::FindOrAdd(EpicAccountId, UserId);
-	const FString& NetId = UserNetId->ToString();
+	LocalUser.UniqueNetId = UserNetId;
+
 	FUserOnlineAccountEOSRef UserAccountRef(new FUserOnlineAccountEOS(UserNetId));
+	LocalUser.UserOnlineAccount = UserAccountRef;
 
-	LocalUserUniqueNetIds.Insert(LocalUserNum, UserNetId);
-	UniqueNetIdToUserAccountMap.Emplace(UserNetId, UserAccountRef);
-	EpicAccountIdToAttributeAccessMap.Emplace(EpicAccountId, UserAccountRef);
+	LocalUser.FriendsList = MakeShareable(new FFriendsListEOS(LocalUserNum, UserNetId));
 
-	// Init player lists
-	FFriendsListEOSRef FriendsList = MakeShareable(new FFriendsListEOS(LocalUserNum, UserNetId));
-	LocalUserNumToFriendsListMap.Emplace(LocalUserNum, FriendsList);
-	NetIdStringToFriendsListMap.Emplace(NetId, FriendsList);
+	LocalUser.LastLoginCredentials = MakeShared<FOnlineAccountCredentials>(AccountCredentials);
+
+	// Add auth refresh notification if not set for this user yet
+	if (!LocalUser.ConnectLoginNotification.IsValid())
+	{
+		FNotificationIdCallbackPairPtr NotificationPairPtr = MakeShared<FNotificationIdCallbackPair>();
+		LocalUser.ConnectLoginNotification = NotificationPairPtr;
+
+		FRefreshAuthCallback* CallbackObj = new FRefreshAuthCallback(AsWeak());
+		NotificationPairPtr->Callback = CallbackObj;
+		CallbackObj->CallbackLambda = [LocalUserNum, this](const EOS_Connect_AuthExpirationCallbackInfo* Data)
+		{
+			RefreshConnectLogin(LocalUserNum);
+		};
+
+		EOS_Connect_AddNotifyAuthExpirationOptions Options = { };
+		Options.ApiVersion = 1;
+		UE_EOS_CHECK_API_MISMATCH(EOS_CONNECT_ADDNOTIFYAUTHEXPIRATION_API_LATEST, 1);
+		NotificationPairPtr->NotificationId = EOS_Connect_AddNotifyAuthExpiration(EOSSubsystem->ConnectHandle, &Options, CallbackObj, CallbackObj->GetCallbackPtr());
+	}
+
+	UniqueNetIdToAttributeAccessMap.Emplace(UserNetId, UserAccountRef);
+
+	// Once all the fields are set, we start gathering additional information
 	ReadFriendsList(LocalUserNum, FString());
 
 	// Get auth token info
@@ -1067,6 +1072,8 @@ void FUserManagerEOS::AddLocalUser(int32 LocalUserNum, EOS_EpicAccountId EpicAcc
 
 		UpdateUserInfo(UserAccountRef, EpicAccountId, EpicAccountId);
 	}
+
+	return LocalUser;
 }
 
 void FUserManagerEOS::UpdateUserInfo(IAttributeAccessInterfaceRef AttributeAccessRef, EOS_EpicAccountId LocalId, EOS_EpicAccountId AccountId)
@@ -1091,21 +1098,17 @@ void FUserManagerEOS::UpdateUserInfo(IAttributeAccessInterfaceRef AttributeAcces
 
 bool FUserManagerEOS::IsFriendQueryUserInfoOngoing(int32 LocalUserNum)
 {
+	const FLocalUserEOS& LocalUser = GetLocalUserChecked(LocalUserNum);
+
 	// If we have an entry for this user and the corresponding array has any element, users are still being processed
-	if (IsFriendQueryUserInfoOngoingForLocalUserMap.Contains(LocalUserNum))
+	if (LocalUser.OngoingFriendQueryUserInfo.Num() > 0)
 	{
-		if(IsFriendQueryUserInfoOngoingForLocalUserMap[LocalUserNum].Num() > 0)
-		{
-			return true;
-		}
+		return true;
 	}
 	
-	if (IsPlayerQueryExternalMappingsOngoingForLocalUserMap.Contains(LocalUserNum))
+	if (LocalUser.OngoingPlayerQueryExternalMappings.Num() > 0)
 	{
-		if(IsPlayerQueryExternalMappingsOngoingForLocalUserMap[LocalUserNum].Num() > 0)
-		{
-			return true;
-		}
+		return true;
 	}
 
 	return false;
@@ -1115,24 +1118,37 @@ TSharedPtr<FUserOnlineAccount> FUserManagerEOS::GetUserAccount(const FUniqueNetI
 {
 	TSharedPtr<FUserOnlineAccount> Result;
 
-	const FUniqueNetIdEOS& EOSID = FUniqueNetIdEOS::Cast(UserId);
-	const FUserOnlineAccountEOSRef* FoundUserAccount = UniqueNetIdToUserAccountMap.Find(EOSID.AsShared());
-	if (FoundUserAccount != nullptr)
+	for (int32 Index = 0; Index < LocalUsers.GetMaxIndex(); Index++)
 	{
-		return *FoundUserAccount;
+		if (LocalUsers.IsValidIndex(Index))
+		{
+			const FLocalUserEOS& LocalUser = LocalUsers[Index];
+
+			if (*LocalUser.UniqueNetId == UserId && LocalUser.UserOnlineAccount.IsValid())
+			{
+				Result = LocalUser.UserOnlineAccount;
+			}
+		}
 	}
 
-	return nullptr;
+	return Result;
 }
 
 TArray<TSharedPtr<FUserOnlineAccount>> FUserManagerEOS::GetAllUserAccounts() const
 {
 	TArray<TSharedPtr<FUserOnlineAccount>> Result;
 
-	for (TUniqueNetIdMap<FUserOnlineAccountEOSRef>::TConstIterator It(UniqueNetIdToUserAccountMap); It; ++It)
+	for (int32 Index = 0; Index < LocalUsers.GetMaxIndex(); Index++)
 	{
-		Result.Add(It.Value());
+		if (LocalUsers.IsValidIndex(Index))
+		{
+			if (const FUserOnlineAccountEOSPtr& UserOnlineAccount = LocalUsers[Index].UserOnlineAccount)
+			{
+				Result.Add(UserOnlineAccount);
+			}
+		}
 	}
+
 	return Result;
 }
 
@@ -1143,9 +1159,9 @@ FUniqueNetIdPtr FUserManagerEOS::GetUniquePlayerId(int32 LocalUserNum) const
 
 const FUniqueNetIdEOSPtr FUserManagerEOS::GetLocalUniqueNetIdEOS(int32 LocalUserNum) const
 {
-	if (LocalUserUniqueNetIds.IsValidIndex(LocalUserNum))
+	if (LocalUsers.IsValidIndex(LocalUserNum))
 	{
-		return LocalUserUniqueNetIds[LocalUserNum];
+		return LocalUsers[LocalUserNum].UniqueNetId;
 	}
 
 	return nullptr;
@@ -1175,11 +1191,11 @@ int32 FUserManagerEOS::GetLocalUserNumFromUniqueNetId(const FUniqueNetId& NetId)
 {
 	const FUniqueNetIdEOS& EosId = FUniqueNetIdEOS::Cast(NetId);
 
-	for (int32 LocalUserNum = 0; LocalUserNum < LocalUserUniqueNetIds.GetMaxIndex(); LocalUserNum++)
+	for (int32 LocalUserNum = 0; LocalUserNum < LocalUsers.GetMaxIndex(); LocalUserNum++)
 	{
-		if (LocalUserUniqueNetIds.IsValidIndex(LocalUserNum))
+		if (LocalUsers.IsValidIndex(LocalUserNum))
 		{
-			const FUniqueNetIdEOSRef& Entry = LocalUserUniqueNetIds[LocalUserNum];
+			const FUniqueNetIdEOSPtr& Entry = LocalUsers[LocalUserNum].UniqueNetId;
 
 			if (*Entry == EosId)
 			{
@@ -1193,11 +1209,11 @@ int32 FUserManagerEOS::GetLocalUserNumFromUniqueNetId(const FUniqueNetId& NetId)
 
 int32 FUserManagerEOS::GetLocalUserNumFromEpicAccountId(const EOS_EpicAccountId& EpicAccountId) const
 {
-	for (int32 LocalUserNum = 0; LocalUserNum < LocalUserUniqueNetIds.GetMaxIndex(); LocalUserNum++)
+	for (int32 LocalUserNum = 0; LocalUserNum < LocalUsers.GetMaxIndex(); LocalUserNum++)
 	{
-		if (LocalUserUniqueNetIds.IsValidIndex(LocalUserNum))
+		if (LocalUsers.IsValidIndex(LocalUserNum))
 		{
-			const FUniqueNetIdEOSRef& Entry = LocalUserUniqueNetIds[LocalUserNum];
+			const FUniqueNetIdEOSPtr& Entry = LocalUsers[LocalUserNum].UniqueNetId;
 
 			if (Entry->GetEpicAccountId() == EpicAccountId)
 			{
@@ -1211,11 +1227,11 @@ int32 FUserManagerEOS::GetLocalUserNumFromEpicAccountId(const EOS_EpicAccountId&
 
 int32 FUserManagerEOS::GetLocalUserNumFromProductUserId(const EOS_ProductUserId& ProductUserId) const
 {
-	for (int32 LocalUserNum = 0; LocalUserNum < LocalUserUniqueNetIds.GetMaxIndex(); LocalUserNum++)
+	for (int32 LocalUserNum = 0; LocalUserNum < LocalUsers.GetMaxIndex(); LocalUserNum++)
 	{
-		if (LocalUserUniqueNetIds.IsValidIndex(LocalUserNum))
+		if (LocalUsers.IsValidIndex(LocalUserNum))
 		{
-			const FUniqueNetIdEOSRef& Entry = LocalUserUniqueNetIds[LocalUserNum];
+			const FUniqueNetIdEOSPtr& Entry = LocalUsers[LocalUserNum].UniqueNetId;
 
 			if (Entry->GetProductUserId() == ProductUserId)
 			{
@@ -1229,11 +1245,16 @@ int32 FUserManagerEOS::GetLocalUserNumFromProductUserId(const EOS_ProductUserId&
 
 const FUniqueNetIdEOSPtr FUserManagerEOS::GetLocalUniqueNetIdEOS(const EOS_EpicAccountId& EpicAccountId) const
 {
-	for (const FUniqueNetIdEOSRef& NetId : LocalUserUniqueNetIds)
+	for (int32 LocalUserNum = 0; LocalUserNum < LocalUsers.GetMaxIndex(); LocalUserNum++)
 	{
-		if (NetId->GetEpicAccountId() == EpicAccountId)
+		if (LocalUsers.IsValidIndex(LocalUserNum))
 		{
-			return NetId;
+			const FUniqueNetIdEOSPtr& Entry = LocalUsers[LocalUserNum].UniqueNetId;
+
+			if (Entry->GetEpicAccountId() == EpicAccountId)
+			{
+				return Entry;
+			}
 		}
 	}
 
@@ -1242,11 +1263,16 @@ const FUniqueNetIdEOSPtr FUserManagerEOS::GetLocalUniqueNetIdEOS(const EOS_EpicA
 
 const FUniqueNetIdEOSPtr FUserManagerEOS::GetLocalUniqueNetIdEOS(const EOS_ProductUserId& ProductUserId) const
 {
-	for (const FUniqueNetIdEOSRef& NetId : LocalUserUniqueNetIds)
+	for (int32 LocalUserNum = 0; LocalUserNum < LocalUsers.GetMaxIndex(); LocalUserNum++)
 	{
-		if (NetId->GetProductUserId() == ProductUserId)
+		if (LocalUsers.IsValidIndex(LocalUserNum))
 		{
-			return NetId;
+			const FUniqueNetIdEOSPtr& Entry = LocalUsers[LocalUserNum].UniqueNetId;
+
+			if (Entry->GetProductUserId() == ProductUserId)
+			{
+				return Entry;
+			}
 		}
 	}
 
@@ -1386,11 +1412,13 @@ FOnlineUserPtr FUserManagerEOS::GetLocalOnlineUser(int32 LocalUserNum) const
 {
 	FOnlineUserPtr OnlineUser;
 
-	if (const FUniqueNetIdEOSPtr NetIdPtr = GetLocalUniqueNetIdEOS(LocalUserNum))
+	if (LocalUsers.IsValidIndex(LocalUserNum))
 	{
-		if (const FUserOnlineAccountEOSRef* Entry = UniqueNetIdToUserAccountMap.Find(NetIdPtr.ToSharedRef()))
+		const FLocalUserEOS& LocalUser = LocalUsers[LocalUserNum];
+
+		if (LocalUser.UserOnlineAccount.IsValid())
 		{
-			OnlineUser = *Entry;
+			OnlineUser = LocalUser.UserOnlineAccount;
 		}
 	}
 
@@ -1401,44 +1429,53 @@ FOnlineUserPtr FUserManagerEOS::GetOnlineUser(EOS_ProductUserId UserId) const
 {
 	FOnlineUserPtr OnlineUser;
 
-	if (const FUniqueNetIdEOSPtr NetIdPtr = FUniqueNetIdEOSRegistry::Find(UserId))
+	for (int32 Index = 0; Index < LocalUsers.GetMaxIndex(); Index++)
 	{
-		if (const FUserOnlineAccountEOSRef* Entry = UniqueNetIdToUserAccountMap.Find(NetIdPtr.ToSharedRef()))
+		if (LocalUsers.IsValidIndex(Index))
 		{
-			OnlineUser = *Entry;
+			const FLocalUserEOS& LocalUser = LocalUsers[Index];
+
+			if (LocalUser.UniqueNetId->GetProductUserId() == UserId && LocalUser.UserOnlineAccount.IsValid())
+			{
+				OnlineUser = LocalUser.UserOnlineAccount;
+			}
 		}
 	}
 
-	return OnlineUser;
+	return OnlineUser;	
 }
 
 FOnlineUserPtr FUserManagerEOS::GetOnlineUser(EOS_EpicAccountId AccountId) const
 {
 	FOnlineUserPtr OnlineUser;
 
-	if (const FUniqueNetIdEOSPtr NetIdPtr = FUniqueNetIdEOSRegistry::Find(AccountId))
+	for (int32 Index = 0; Index < LocalUsers.GetMaxIndex(); Index++)
 	{
-		if (const FUserOnlineAccountEOSRef* Entry = UniqueNetIdToUserAccountMap.Find(NetIdPtr.ToSharedRef()))
+		if (LocalUsers.IsValidIndex(Index))
 		{
-			OnlineUser = *Entry;
+			const FLocalUserEOS& LocalUser = LocalUsers[Index];
+
+			if (LocalUser.UniqueNetId->GetEpicAccountId() == AccountId && LocalUser.UserOnlineAccount.IsValid())
+			{
+				OnlineUser = LocalUser.UserOnlineAccount;
+			}
 		}
 	}
-	
+
 	return OnlineUser;
 }
 
 void FUserManagerEOS::RemoveLocalUser(int32 LocalUserNum)
 {
-	if (LocalUserUniqueNetIds.IsValidIndex(LocalUserNum))
+	if (LocalUsers.IsValidIndex(LocalUserNum))
 	{
-		const FUniqueNetIdEOSRef FoundId = LocalUserUniqueNetIds[LocalUserNum];
+		const FUniqueNetIdEOSRef FoundId = GetLocalUserChecked(LocalUserNum).UniqueNetId.ToSharedRef();
 
 		EOSSubsystem->ReleaseVoiceChatUserInterface(*FoundId);
-		LocalUserNumToFriendsListMap.Remove(LocalUserNum);
-		const FString& NetId = FoundId->ToString();
-		LocalUserUniqueNetIds.RemoveAt(LocalUserNum);
-		UniqueNetIdToUserAccountMap.Remove(FoundId);
+
+		LocalUsers.RemoveAt(LocalUserNum);
 	}
+
 	// Reset this for the next user login
 	if (LocalUserNum == DefaultLocalUser)
 	{
@@ -1732,42 +1769,40 @@ void FUserManagerEOS::FriendStatusChanged(const EOS_Friends_OnFriendsUpdateInfo*
 	if (IsLocalUser(Data->LocalUserId))
 	{
 		const int32 LocalUserNum = GetLocalUserNumFromEpicAccountId(Data->LocalUserId);
-		const FUniqueNetIdEOSRef LocalEOSID = LocalUserUniqueNetIds[LocalUserNum];
+		const FUniqueNetIdEOSRef LocalEOSID = GetLocalUserChecked(LocalUserNum).UniqueNetId.ToSharedRef();
 		if (!FUniqueNetIdEOSRegistry::Find(Data->TargetUserId))
 		{
-			FUniqueNetIdEOSRef FriendNetId = FUniqueNetIdEOSRegistry::FindOrAdd(Data->TargetUserId, nullptr);
-			const FString NetId = FriendNetId->ToString();
-			AddRemotePlayer(LocalUserNum, NetId, Data->TargetUserId);
+			AddRemotePlayer(LocalUserNum, Data->TargetUserId);
 			// TODO: We need to modify AddRemotePlayer to take a delegate, to execute after all information has been retrieved
 		}
 
 		// AddRemotePlayer just added this id to the registry, so it will be found and valid, no checks needed
-		const FUniqueNetIdEOS& FriendEOSId = *FUniqueNetIdEOSRegistry::FindChecked(Data->TargetUserId);
+		const FUniqueNetIdEOSRef& FriendEOSId = FUniqueNetIdEOSRegistry::FindChecked(Data->TargetUserId);
 
 		switch (Data->CurrentStatus)
 		{
 		case EOS_EFriendsStatus::EOS_FS_NotFriends: // Invite rejections and friend removal
 		{
 			//User should already be a friend
-			FFriendsListEOSRef& FriendsListRef = LocalUserNumToFriendsListMap[LocalUserNum];
-			FOnlineFriendEOSPtr Friend = FriendsListRef->GetByNetIdString(FriendEOSId.ToString());
+			FFriendsListEOSPtr& FriendsListRef = GetLocalUserChecked(LocalUserNum).FriendsList;
+			FOnlineFriendEOSPtr Friend = FriendsListRef->GetByNetId(FriendEOSId);
 			if (Friend.IsValid())
 			{
-				FriendsListRef->Remove(FriendEOSId.ToString(), Friend.ToSharedRef());
+				FriendsListRef->Remove(FriendEOSId, Friend.ToSharedRef());
 				Friend->SetInviteStatus(EInviteStatus::Unknown);
 
 				if (Data->PreviousStatus == EOS_EFriendsStatus::EOS_FS_Friends)
 				{
-					TriggerOnFriendRemovedDelegates(*LocalEOSID, FriendEOSId);
+					TriggerOnFriendRemovedDelegates(*LocalEOSID, *FriendEOSId);
 				}
 				else if (Data->PreviousStatus == EOS_EFriendsStatus::EOS_FS_InviteSent || Data->PreviousStatus == EOS_EFriendsStatus::EOS_FS_InviteReceived)
 				{
-					TriggerOnInviteRejectedDelegates(*LocalEOSID, FriendEOSId); // We don't have an "OnInviteRejected" event only for the local user
+					TriggerOnInviteRejectedDelegates(*LocalEOSID, *FriendEOSId); // We don't have an "OnInviteRejected" event only for the local user
 				}
 			}
 			else
 			{
-				UE_LOG_ONLINE_FRIEND(Verbose, TEXT("Friend status notification received for user [%d], but remote user [%s] was not previously registered as a friend"), *FriendEOSId.ToString());
+				UE_LOG_ONLINE_FRIEND(Verbose, TEXT("Friend status notification received for user [%d], but remote user [%s] was not previously registered as a friend"), *FriendEOSId->ToString());
 			}
 
 			break;
@@ -1792,7 +1827,7 @@ void FUserManagerEOS::FriendStatusChanged(const EOS_Friends_OnFriendsUpdateInfo*
 			FOnlineFriendEOSRef FriendRef = AddFriend(LocalUserNum, Data->TargetUserId);
 
 			FriendRef->SetInviteStatus(EInviteStatus::PendingInbound);
-			TriggerOnInviteReceivedDelegates(*LocalEOSID, FriendEOSId);
+			TriggerOnInviteReceivedDelegates(*LocalEOSID, *FriendEOSId);
 
 			break;
 		}
@@ -1801,7 +1836,7 @@ void FUserManagerEOS::FriendStatusChanged(const EOS_Friends_OnFriendsUpdateInfo*
 			FOnlineFriendEOSRef FriendRef = AddFriend(LocalUserNum, Data->TargetUserId);
 
 			FriendRef->SetInviteStatus(EInviteStatus::Accepted);
-			TriggerOnInviteAcceptedDelegates(*LocalEOSID, FriendEOSId); // We don't have an "OnFriendAdded", all relationships exist in the context of invites. We also don't have a separate "OnInviteAccepted" for the local user.
+			TriggerOnInviteAcceptedDelegates(*LocalEOSID, *FriendEOSId); // We don't have an "OnFriendAdded", all relationships exist in the context of invites. We also don't have a separate "OnInviteAccepted" for the local user.
 
 			break;
 		}
@@ -1817,9 +1852,9 @@ void FUserManagerEOS::FriendStatusChanged(const EOS_Friends_OnFriendsUpdateInfo*
 FOnlineFriendEOSRef FUserManagerEOS::AddFriend(int32 LocalUserNum, EOS_EpicAccountId EpicAccountId)
 {
 	FUniqueNetIdEOSRef FriendNetId = FUniqueNetIdEOSRegistry::FindOrAdd(EpicAccountId, nullptr);
-	const FString NetId = FriendNetId->ToString();
 	FOnlineFriendEOSRef FriendRef = MakeShareable(new FOnlineFriendEOS(FriendNetId));
-	LocalUserNumToFriendsListMap[LocalUserNum]->Add(NetId, FriendRef);
+
+	GetLocalUserChecked(LocalUserNum).FriendsList->Add(FriendNetId, FriendRef);
 
 	EOS_Friends_GetStatusOptions Options = { };
 	Options.ApiVersion = 1;
@@ -1833,7 +1868,7 @@ FOnlineFriendEOSRef FUserManagerEOS::AddFriend(int32 LocalUserNum, EOS_EpicAccou
 	// Add this friend as a remote player if we haven't already (this will grab user info)
 	if (!FUniqueNetIdEOSRegistry::Find(EpicAccountId))
 	{
-		AddRemotePlayer(LocalUserNum, NetId, EpicAccountId, FriendNetId, FriendRef, FriendRef);
+		AddRemotePlayer(LocalUserNum, EpicAccountId, FriendRef);
 	}
 
 	// Querying the presence of a non-friend would cause an SDK error.
@@ -1847,21 +1882,20 @@ FOnlineFriendEOSRef FUserManagerEOS::AddFriend(int32 LocalUserNum, EOS_EpicAccou
 	return FriendRef;
 }
 
-void FUserManagerEOS::AddRemotePlayer(int32 LocalUserNum, const FString& NetId, EOS_EpicAccountId EpicAccountId)
+void FUserManagerEOS::AddRemotePlayer(int32 LocalUserNum, EOS_EpicAccountId EpicAccountId)
 {
-	FUniqueNetIdEOSRef EOSID = FUniqueNetIdEOSRegistry::FindOrAdd(NetId);
+	FUniqueNetIdEOSRef EOSID = FUniqueNetIdEOSRegistry::FindOrAdd(EpicAccountId, nullptr);
 	FOnlineUserEOSRef UserRef = MakeShareable(new FOnlineUserEOS(EOSID));
+
 	// Add this user as a remote (this will grab presence & user info)
-	AddRemotePlayer(LocalUserNum, NetId, EpicAccountId, EOSID, UserRef, UserRef);
+	AddRemotePlayer(LocalUserNum, EpicAccountId, UserRef);
 }
 
-void FUserManagerEOS::AddRemotePlayer(int32 LocalUserNum, const FString& NetId, EOS_EpicAccountId EpicAccountId, FUniqueNetIdEOSPtr UniqueNetId, FOnlineUserPtr OnlineUser, IAttributeAccessInterfaceRef AttributeRef)
+void FUserManagerEOS::AddRemotePlayer(int32 LocalUserNum, EOS_EpicAccountId EpicAccountId, IAttributeAccessInterfaceRef AttributeRef)
 {
 	// TODO: pending asynchronous resolution of the remote player's UniqueNetId
-	FUniqueNetIdEOSRegistry::FindOrAdd(EpicAccountId, nullptr);
-	
-	NetIdStringToAttributeAccessMap.Emplace(NetId, AttributeRef);
-	EpicAccountIdToAttributeAccessMap.Emplace(EpicAccountId, AttributeRef);
+	FUniqueNetIdEOSRef EOSID = FUniqueNetIdEOSRegistry::FindOrAdd(EpicAccountId, nullptr);
+	UniqueNetIdToAttributeAccessMap.Emplace(EOSID, AttributeRef);
 
 	// Read the user info for this player
 	ReadUserInfo(LocalUserNum, EpicAccountId);
@@ -1869,54 +1903,16 @@ void FUserManagerEOS::AddRemotePlayer(int32 LocalUserNum, const FString& NetId,
 
 void FUserManagerEOS::UpdateRemotePlayerProductUserId(EOS_EpicAccountId EpicAccountId, EOS_ProductUserId ProductUserId)
 {
-	const FUniqueNetIdEOSRef NewNetId = FUniqueNetIdEOSRegistry::FindOrAdd(EpicAccountId, ProductUserId);
-
-	const FString PrevNetIdStr = LexToString(EpicAccountId) + EOS_ID_SEPARATOR;
-	// Note that calling FindOrAdd with previously invalid EAS/EOS components will update the net id in place.
-	const FString NewNetIdStr = NewNetId->ToString();
-	
-	// TODO: While there is still unresolved ids and string maps in the system, we'll force this method to always update
-	/*if (PrevNetIdStr == NewNetIdStr)
-	{
-		// No change, so skip any work
-		return;
-	}*/
+	// TODO: Unique Id updates are removed, but we still need to add proper id resolution for all ids to keep a correct state
 
-	const FString ProductUserIdStr = LexToString(ProductUserId);
-
-	// Update any old friends entries with the new net id key
-	for (TMap<int32, FFriendsListEOSRef>::TConstIterator It(LocalUserNumToFriendsListMap); It; ++It)
-	{
-		FFriendsListEOSRef FriendsList = It.Value();
-		FOnlineFriendEOSPtr FoundFriend = FriendsList->GetByNetIdString(PrevNetIdStr);
-		if (FoundFriend.IsValid())
-		{
-			FriendsList->UpdateNetIdStr(PrevNetIdStr, NewNetIdStr);
-		}
-	}
-
-	// Update all of the other net id to X mappings
-	if (NetIdStringToAttributeAccessMap.Contains(PrevNetIdStr))
-	{
-		const IAttributeAccessInterfaceRef AttrAccess = NetIdStringToAttributeAccessMap[PrevNetIdStr];
-		NetIdStringToAttributeAccessMap.Remove(PrevNetIdStr);
-		NetIdStringToAttributeAccessMap.Emplace(NewNetIdStr, AttrAccess);
-	}
-
-	// Presence may not be available for all online users
-	if (NetIdStringToOnlineUserPresenceMap.Contains(PrevNetIdStr))
-	{
-		FOnlineUserPresenceRef UserPresence = NetIdStringToOnlineUserPresenceMap[PrevNetIdStr];
-		NetIdStringToOnlineUserPresenceMap.Remove(PrevNetIdStr);
-		NetIdStringToOnlineUserPresenceMap.Emplace(NewNetIdStr, UserPresence);
-	}
+	FUniqueNetIdEOSRegistry::FindOrAdd(EpicAccountId, ProductUserId);	
 }
 
 // IOnlineFriends Interface
 
 bool FUserManagerEOS::ReadFriendsList(int32 LocalUserNum, const FString& ListName, const FOnReadFriendsListComplete& Delegate)
 {
-	if (!LocalUserUniqueNetIds.IsValidIndex(LocalUserNum))
+	if (!LocalUsers.IsValidIndex(LocalUserNum))
 	{
 		const FString ErrorStr = FString::Printf(TEXT("Can't ReadFriendsList() for user (%d) since they are not logged in"), LocalUserNum);
 		UE_LOG_ONLINE_FRIEND(Warning, TEXT("%s"), *ErrorStr);
@@ -1925,12 +1921,9 @@ bool FUserManagerEOS::ReadFriendsList(int32 LocalUserNum, const FString& ListNam
 	}
 
 	// We save the information for this call even if it won't be automatically processed
-	const bool bIsReadFriendsListOngoing = CachedReadUserListInfoForLocalUserMap.Contains(LocalUserNum);
-	if (!bIsReadFriendsListOngoing)
-	{
-		CachedReadUserListInfoForLocalUserMap.Emplace(LocalUserNum);
-	}
-	CachedReadUserListInfoForLocalUserMap[LocalUserNum].Add(ReadUserListInfo(LocalUserNum, ListName, Delegate));
+	FLocalUserEOS& LocalUser = GetLocalUserChecked(LocalUserNum);
+	const bool bIsReadFriendsListOngoing = !LocalUser.CachedReadUserListInfo.IsEmpty();
+	LocalUser.CachedReadUserListInfo.Add(ReadUserListInfo(LocalUserNum, ListName, Delegate));
 
 	if (bIsReadFriendsListOngoing)
 	{
@@ -1962,7 +1955,7 @@ bool FUserManagerEOS::ReadFriendsList(int32 LocalUserNum, const FString& ListNam
 			Options.LocalUserId = GetLocalEpicAccountId(LocalUserNum);
 			int32 FriendCount = EOS_Friends_GetFriendsCount(EOSSubsystem->FriendsHandle, &Options);
 
-			LocalUserNumToFriendsListMap[LocalUserNum]->Empty(FriendCount);
+			GetLocalUserChecked(LocalUserNum).FriendsList->Empty(FriendCount);
 
 			TArray<FString> FriendEasIds;
 			FriendEasIds.Reserve(FriendCount);
@@ -2012,20 +2005,20 @@ void FUserManagerEOS::ProcessReadFriendsListComplete(int32 LocalUserNum, bool bW
 	// If we started any user info queries for friends, we'll just wait until they finish
 	if (!IsFriendQueryUserInfoOngoing(LocalUserNum))
 	{
+		FLocalUserEOS& LocalUser = GetLocalUserChecked(LocalUserNum);
+
 		// If not, we'll just trigger the delegates for all cached calls
-		TArray<ReadUserListInfo> CachedInfoList;
-		if (CachedReadUserListInfoForLocalUserMap.RemoveAndCopyValue(LocalUserNum, CachedInfoList))
+		for (const ReadUserListInfo& CachedInfo : LocalUsers[LocalUserNum].CachedReadUserListInfo)
 		{
-			for (const ReadUserListInfo& CachedInfo : CachedInfoList)
-			{
-				CachedInfo.ExecuteDelegateIfBound(bWasSuccessful, ErrorStr);
-			}
+			CachedInfo.ExecuteDelegateIfBound(bWasSuccessful, ErrorStr);
 		}
+		LocalUser.CachedReadUserListInfo.Empty();
 
-		TriggerOnFriendsChangeDelegates(LocalUserNum);
+		// These should be empty by now, but we make sure
+		LocalUser.OngoingFriendQueryUserInfo.Empty();
+		LocalUser.OngoingPlayerQueryExternalMappings.Empty();
 
-		IsFriendQueryUserInfoOngoingForLocalUserMap.Remove(LocalUserNum);
-		IsPlayerQueryExternalMappingsOngoingForLocalUserMap.Remove(LocalUserNum);
+		TriggerOnFriendsChangeDelegates(LocalUserNum);
 	}
 }
 
@@ -2065,7 +2058,7 @@ typedef TEOSCallback<EOS_Friends_OnSendInviteCallback, EOS_Friends_SendInviteCal
 
 bool FUserManagerEOS::SendInvite(int32 LocalUserNum, const FUniqueNetId& FriendId, const FString& ListName, const FOnSendInviteComplete& Delegate)
 {
-	if (!LocalUserUniqueNetIds.IsValidIndex(LocalUserNum))
+	if (!LocalUsers.IsValidIndex(LocalUserNum))
 	{
 		UE_LOG_ONLINE_FRIEND(Warning, TEXT("Can't SendInvite() for user (%d) since they are not logged in"), LocalUserNum);
 		Delegate.ExecuteIfBound(LocalUserNum, false, FriendId, ListName, FString(TEXT("Can't SendInvite() for user (%d) since they are not logged in"), LocalUserNum));
@@ -2110,7 +2103,7 @@ typedef TEOSCallback<EOS_Friends_OnAcceptInviteCallback, EOS_Friends_AcceptInvit
 
 bool FUserManagerEOS::AcceptInvite(int32 LocalUserNum, const FUniqueNetId& FriendId, const FString& ListName, const FOnAcceptInviteComplete& Delegate)
 {
-	if (!LocalUserUniqueNetIds.IsValidIndex(LocalUserNum))
+	if (!LocalUsers.IsValidIndex(LocalUserNum))
 	{
 		UE_LOG_ONLINE_FRIEND(Warning, TEXT("Can't AcceptInvite() for user (%d) since they are not logged in"), LocalUserNum);
 		Delegate.ExecuteIfBound(LocalUserNum, false, FriendId, ListName, FString(TEXT("Can't AcceptInvite() for user (%d) since they are not logged in"), LocalUserNum));
@@ -2156,7 +2149,7 @@ void EOS_CALL EOSRejectInviteCallback(const EOS_Friends_RejectInviteCallbackInfo
 
 bool FUserManagerEOS::RejectInvite(int32 LocalUserNum, const FUniqueNetId& FriendId, const FString& ListName)
 {
-	if (!LocalUserUniqueNetIds.IsValidIndex(LocalUserNum))
+	if (!LocalUsers.IsValidIndex(LocalUserNum))
 	{
 		UE_LOG_ONLINE_FRIEND(Warning, TEXT("Can't RejectInvite() for user (%d) since they are not logged in"), LocalUserNum);
 		return false;
@@ -2197,81 +2190,77 @@ bool FUserManagerEOS::DeleteFriend(int32 LocalUserNum, const FUniqueNetId& Frien
 bool FUserManagerEOS::GetFriendsList(int32 LocalUserNum, const FString& ListName, TArray<TSharedRef<FOnlineFriend>>& OutFriends)
 {
 	OutFriends.Reset();
-	if (LocalUserNumToFriendsListMap.Contains(LocalUserNum))
+
+	for (FOnlineFriendEOSRef Friend : GetLocalUserChecked(LocalUserNum).FriendsList->GetList())
 	{
-		FFriendsListEOSRef FriendsList = LocalUserNumToFriendsListMap[LocalUserNum];
-		for (FOnlineFriendEOSRef Friend : FriendsList->GetList())
+		const FOnlineUserPresence& Presence = Friend->GetPresence();
+		// See if they only want online only
+		if (ListName == EFriendsLists::ToString(EFriendsLists::OnlinePlayers) && !Presence.bIsOnline)
 		{
-			const FOnlineUserPresence& Presence = Friend->GetPresence();
-			// See if they only want online only
-			if (ListName == EFriendsLists::ToString(EFriendsLists::OnlinePlayers) && !Presence.bIsOnline)
-			{
-				continue;
-			}
-			// Of if they only want friends playing this game
-			else if (ListName == EFriendsLists::ToString(EFriendsLists::InGamePlayers) && !Presence.bIsPlayingThisGame)
-			{
-				continue;
-			}
-			// If the service hasn't returned the info yet, skip them
-			else if (Friend->GetDisplayName().IsEmpty())
-			{
-				continue;
-			}
-			OutFriends.Add(Friend);
+			continue;
+		}
+		// Of if they only want friends playing this game
+		else if (ListName == EFriendsLists::ToString(EFriendsLists::InGamePlayers) && !Presence.bIsPlayingThisGame)
+		{
+			continue;
+		}
+		// If the service hasn't returned the info yet, skip them
+		else if (Friend->GetDisplayName().IsEmpty())
+		{
+			continue;
 		}
-		// Sort these by those playing the game first, alphabetically, then not playing, then not online
-		OutFriends.Sort([](TSharedRef<FOnlineFriend> A, TSharedRef<FOnlineFriend> B)
+		OutFriends.Add(Friend);
+	}
+
+	// Sort these by those playing the game first, alphabetically, then not playing, then not online
+	OutFriends.Sort([](TSharedRef<FOnlineFriend> A, TSharedRef<FOnlineFriend> B)
+	{
+		const FOnlineUserPresence& APres = A->GetPresence();
+		const FOnlineUserPresence& BPres = B->GetPresence();
+		// If they are the same, then check playing this game
+		if (APres.bIsOnline == BPres.bIsOnline)
 		{
-			const FOnlineUserPresence& APres = A->GetPresence();
-			const FOnlineUserPresence& BPres = B->GetPresence();
-			// If they are the same, then check playing this game
-			if (APres.bIsOnline == BPres.bIsOnline)
+			// If they are the same, then sort by name
+			if (APres.bIsPlayingThisGame == BPres.bIsPlayingThisGame)
 			{
-				// If they are the same, then sort by name
-				if (APres.bIsPlayingThisGame == BPres.bIsPlayingThisGame)
+				const EInviteStatus::Type AFriendStatus = A->GetInviteStatus();
+				const EInviteStatus::Type BFriendStatus = B->GetInviteStatus();
+				// Sort pending friends below accepted friends
+				if (AFriendStatus == BFriendStatus && AFriendStatus == EInviteStatus::Accepted)
 				{
-					const EInviteStatus::Type AFriendStatus = A->GetInviteStatus();
-					const EInviteStatus::Type BFriendStatus = B->GetInviteStatus();
-					// Sort pending friends below accepted friends
-					if (AFriendStatus == BFriendStatus && AFriendStatus == EInviteStatus::Accepted)
-					{
-						const FString& AName = A->GetDisplayName();
-						const FString& BName = B->GetDisplayName();
-						return AName < BName;
-					}
+					const FString& AName = A->GetDisplayName();
+					const FString& BName = B->GetDisplayName();
+					return AName < BName;
 				}
 			}
-			return false;
-		});
-		return true;
-	}
-	return false;
+		}
+		return false;
+	});
+
+	return true;
 }
 
 TSharedPtr<FOnlineFriend> FUserManagerEOS::GetFriend(int32 LocalUserNum, const FUniqueNetId& FriendId, const FString& ListName)
 {
-	if (LocalUserNumToFriendsListMap.Contains(LocalUserNum))
+	const FUniqueNetIdEOS& EosId = FUniqueNetIdEOS::Cast(FriendId);
+	FOnlineFriendEOSPtr FoundFriend = GetLocalUserChecked(LocalUserNum).FriendsList->GetByNetId(EosId.AsShared());
+	if (FoundFriend.IsValid())
 	{
-		FFriendsListEOSRef FriendsList = LocalUserNumToFriendsListMap[LocalUserNum];
-		const FUniqueNetIdEOS& EosId = FUniqueNetIdEOS::Cast(FriendId);
-		FOnlineFriendEOSPtr FoundFriend = FriendsList->GetByNetIdString(EosId.ToString());
-		if (FoundFriend.IsValid())
+		const FOnlineUserPresence& Presence = FoundFriend->GetPresence();
+		// See if they only want online only
+		if (ListName == EFriendsLists::ToString(EFriendsLists::OnlinePlayers) && !Presence.bIsOnline)
 		{
-			const FOnlineUserPresence& Presence = FoundFriend->GetPresence();
-			// See if they only want online only
-			if (ListName == EFriendsLists::ToString(EFriendsLists::OnlinePlayers) && !Presence.bIsOnline)
-			{
-				return TSharedPtr<FOnlineFriend>();
-			}
-			// Of if they only want friends playing this game
-			else if (ListName == EFriendsLists::ToString(EFriendsLists::InGamePlayers) && !Presence.bIsPlayingThisGame)
-			{
-				return TSharedPtr<FOnlineFriend>();
-			}
-			return FoundFriend;
+			return TSharedPtr<FOnlineFriend>();
 		}
+		// Of if they only want friends playing this game
+		else if (ListName == EFriendsLists::ToString(EFriendsLists::InGamePlayers) && !Presence.bIsPlayingThisGame)
+		{
+			return TSharedPtr<FOnlineFriend>();
+		}
+
+		return FoundFriend;
 	}
+
 	return TSharedPtr<FOnlineFriend>();
 }
 
@@ -2639,16 +2628,16 @@ void FUserManagerEOS::UpdatePresence(int32 LocalUserNum, EOS_EpicAccountId Accou
 	EOS_EResult CopyResult = EOS_Presence_CopyPresence(EOSSubsystem->PresenceHandle, &Options, &PresenceInfo);
 	if (CopyResult == EOS_EResult::EOS_Success)
 	{
-		const FString& NetIdStr = FUniqueNetIdEOSRegistry::FindChecked(AccountId)->ToString();
+		const FUniqueNetIdEOSRef& NetId = FUniqueNetIdEOSRegistry::FindChecked(AccountId);
 
 		// Create it on demand if we don't have one yet
-		if (!NetIdStringToOnlineUserPresenceMap.Contains(NetIdStr))
+		if (!UniqueNetIdToOnlineUserPresenceMap.Contains(NetId))
 		{
 			FOnlineUserPresenceRef PresenceRef = MakeShareable(new FOnlineUserPresence());
-			NetIdStringToOnlineUserPresenceMap.Emplace(NetIdStr, PresenceRef);
+			UniqueNetIdToOnlineUserPresenceMap.Emplace(NetId, PresenceRef);
 		}
 
-		FOnlineUserPresenceRef PresenceRef = NetIdStringToOnlineUserPresenceMap[NetIdStr];
+		FOnlineUserPresenceRef PresenceRef = UniqueNetIdToOnlineUserPresenceMap[NetId];
 		const FString ProductId(UTF8_TO_TCHAR(PresenceInfo->ProductId));
 		const FString ProdVersion(UTF8_TO_TCHAR(PresenceInfo->ProductVersion));
 		const FString Platform(UTF8_TO_TCHAR(PresenceInfo->Platform));
@@ -2670,7 +2659,7 @@ void FUserManagerEOS::UpdatePresence(int32 LocalUserNum, EOS_EpicAccountId Accou
 		}
 
 		// Copy the presence if this is a friend that was updated, so that their data is in sync
-		UpdateFriendPresence(NetIdStr, PresenceRef);
+		UpdateFriendPresence(NetId, PresenceRef);
 
 		EOS_Presence_Info_Release(PresenceInfo);
 	}
@@ -2680,15 +2669,22 @@ void FUserManagerEOS::UpdatePresence(int32 LocalUserNum, EOS_EpicAccountId Accou
 	}
 }
 
-void FUserManagerEOS::UpdateFriendPresence(const FString& FriendId, FOnlineUserPresenceRef Presence)
+void FUserManagerEOS::UpdateFriendPresence(const FUniqueNetIdEOSRef& FriendId, FOnlineUserPresenceRef Presence)
 {
-	for (TMap<int32, FFriendsListEOSRef>::TConstIterator It(LocalUserNumToFriendsListMap); It; ++It)
+	for (int32 Index = 0; Index < LocalUsers.GetMaxIndex(); Index++)
 	{
-		FFriendsListEOSRef FriendsList = It.Value();
-		FOnlineFriendEOSPtr Friend = FriendsList->GetByNetIdString(FriendId);
-		if (Friend.IsValid())
+		if (LocalUsers.IsValidIndex(Index))
 		{
-			Friend->SetPresence(Presence);
+			const FLocalUserEOS& LocalUser = LocalUsers[Index];
+			
+			if (FFriendsListEOSPtr FriendsList = LocalUser.FriendsList)
+			{
+				FOnlineFriendEOSPtr Friend = FriendsList->GetByNetId(FriendId);
+				if (Friend.IsValid())
+				{
+					Friend->SetPresence(Presence);
+				}
+			}
 		}
 	}
 }
@@ -2696,9 +2692,9 @@ void FUserManagerEOS::UpdateFriendPresence(const FString& FriendId, FOnlineUserP
 EOnlineCachedResult::Type FUserManagerEOS::GetCachedPresence(const FUniqueNetId& UserId, TSharedPtr<FOnlineUserPresence>& OutPresence)
 {
 	const FUniqueNetIdEOS& EOSID = FUniqueNetIdEOS::Cast(UserId);
-	if (NetIdStringToOnlineUserPresenceMap.Contains(EOSID.ToString()))
+	if (UniqueNetIdToOnlineUserPresenceMap.Contains(EOSID.AsShared()))
 	{
-		OutPresence = NetIdStringToOnlineUserPresenceMap[EOSID.ToString()];
+		OutPresence = UniqueNetIdToOnlineUserPresenceMap[EOSID.AsShared()];
 		return EOnlineCachedResult::Success;
 	}
 	return EOnlineCachedResult::NotFound;
@@ -2728,7 +2724,7 @@ bool FUserManagerEOS::QueryUserInfo(int32 LocalUserNum, const TArray<FUniqueNetI
 		if (EOS_EpicAccountId_IsValid(AccountId) == EOS_TRUE)
 		{
 			// If the user is already registered, we'll update their user info
-			if (EpicAccountIdToAttributeAccessMap.Contains(AccountId))
+			if (UniqueNetIdToAttributeAccessMap.Contains(EOSID.AsShared()))
 			{
 				ReadUserInfo(LocalUserNum, AccountId);
 			}
@@ -2738,7 +2734,7 @@ bool FUserManagerEOS::QueryUserInfo(int32 LocalUserNum, const TArray<FUniqueNetI
 				UserEasIdsNeedingExternalMappings.Add(LexToString(AccountId));
 
 				// Registering the player will also query the user info data
-				AddRemotePlayer(LocalUserNum, EOSID.ToString(), AccountId);
+				AddRemotePlayer(LocalUserNum, AccountId);
 			}
 		}
 	}
@@ -2776,12 +2772,13 @@ void FUserManagerEOS::ReadUserInfo(int32 LocalUserNum, EOS_EpicAccountId EpicAcc
 	{
 		if (Data->ResultCode == EOS_EResult::EOS_Success)
 		{
-			IAttributeAccessInterfaceRef AttributeAccessRef = EpicAccountIdToAttributeAccessMap[Data->TargetUserId];
+			const FUniqueNetIdEOSRef EOSId = FUniqueNetIdEOSRegistry::FindChecked(Data->TargetUserId);
+			IAttributeAccessInterfaceRef AttributeAccessRef = UniqueNetIdToAttributeAccessMap[EOSId];
 			UpdateUserInfo(AttributeAccessRef, Data->LocalUserId, Data->TargetUserId);
 		}
 
 		// We mark this player as processed
-		IsFriendQueryUserInfoOngoingForLocalUserMap[LocalUserNum].RemoveSwap(EpicAccountId, false);
+		GetLocalUserChecked(LocalUserNum).OngoingFriendQueryUserInfo.RemoveSwap(EpicAccountId, false);
 
 		ProcessReadFriendsListComplete(LocalUserNum, true, TEXT(""));
 	};
@@ -2794,19 +2791,21 @@ void FUserManagerEOS::ReadUserInfo(int32 LocalUserNum, EOS_EpicAccountId EpicAcc
 	EOS_UserInfo_QueryUserInfo(EOSSubsystem->UserInfoHandle, &Options, CallbackObj, CallbackObj->GetCallbackPtr());
 
 	// We mark this player as pending for processing
-	IsFriendQueryUserInfoOngoingForLocalUserMap.FindOrAdd(LocalUserNum).Add(EpicAccountId);
+	GetLocalUserChecked(LocalUserNum).OngoingFriendQueryUserInfo.Add(EpicAccountId);
 }
 
 bool FUserManagerEOS::GetAllUserInfo(int32 LocalUserNum, TArray<TSharedRef<FOnlineUser>>& OutUsers)
 {
 	OutUsers.Reset();
 
-	// Get local users only
-	for (const FUniqueNetIdEOSRef& LocalUserNetId : LocalUserUniqueNetIds)
+	for (int32 Index = 0; Index < LocalUsers.GetMaxIndex(); Index++)
 	{
-		if (const FUserOnlineAccountEOSRef* LocalUser = UniqueNetIdToUserAccountMap.Find(LocalUserNetId))
+		if (LocalUsers.IsValidIndex(Index))
 		{
-			OutUsers.Add(*LocalUser);
+			if (const FUserOnlineAccountEOSPtr UserOnlineAccount = LocalUsers[Index].UserOnlineAccount)
+			{
+				OutUsers.Add(UserOnlineAccount.ToSharedRef());
+			}
 		}
 	}
 
@@ -2817,9 +2816,17 @@ TSharedPtr<FOnlineUser> FUserManagerEOS::GetUserInfo(int32 LocalUserNum, const F
 {
 	TSharedPtr<FOnlineUser> OnlineUser;
 
-	if (const FUserOnlineAccountEOSRef* FoundUser = UniqueNetIdToUserAccountMap.Find(UserId.AsShared()))
+	for (int32 Index = 0; Index < LocalUsers.GetMaxIndex(); Index++)
 	{
-		OnlineUser = *FoundUser;
+		if (LocalUsers.IsValidIndex(Index))
+		{
+			const FLocalUserEOS& LocalUser = LocalUsers[Index];
+
+			if (*LocalUser.UniqueNetId == UserId && LocalUser.UserOnlineAccount.IsValid())
+			{
+				OnlineUser = LocalUser.UserOnlineAccount;
+			}
+		}
 	}
 
 	return OnlineUser;
@@ -2866,15 +2873,13 @@ bool FUserManagerEOS::QueryUserIdMapping(const FUniqueNetId& UserId, const FStri
 		bool bWasSuccessful = Result == EOS_EResult::EOS_Success;
 		if (bWasSuccessful)
 		{
-			const FString NetIdStr = LexToString(Data->TargetUserId);
-			const FUniqueNetIdEOSRef LocalUserId = LocalUserUniqueNetIds[LocalUserNum];
 			if (!FUniqueNetIdEOSRegistry::Find(Data->TargetUserId))
 			{
 				// Registering the player will also query the presence/user info data
-				AddRemotePlayer(LocalUserNum, NetIdStr, Data->TargetUserId);
+				AddRemotePlayer(LocalUserNum, Data->TargetUserId);
 			}
 
-			Delegate.ExecuteIfBound(true, *LocalUserId, DisplayNameOrEmail, *FUniqueNetIdEOSRegistry::FindOrAdd(NetIdStr), ErrorString);
+			Delegate.ExecuteIfBound(true, *GetLocalUserChecked(LocalUserNum).UniqueNetId, DisplayNameOrEmail, *FUniqueNetIdEOSRegistry::FindOrAdd(Data->TargetUserId, nullptr), ErrorString);
 		}
 		else
 		{
@@ -2955,7 +2960,7 @@ bool FUserManagerEOS::QueryExternalIdMappings(const FUniqueNetId& UserId, const
 	int32 LocalUserNum = GetLocalUserNumFromUniqueNetId(UserId);
 
 	// Mark the queries as in progress
-	IsPlayerQueryExternalMappingsOngoingForLocalUserMap.FindOrAdd(LocalUserNum).Append(ExternalIds);
+	GetLocalUserChecked(LocalUserNum).OngoingPlayerQueryExternalMappings.Append(ExternalIds);
 
 	const EOS_ProductUserId LocalUserId = EOSID.GetProductUserId();
 	const int32 NumBatches = (ExternalIds.Num() / EOS_CONNECT_QUERYEXTERNALACCOUNTMAPPINGS_MAX_ACCOUNT_IDS) + 1;
@@ -2987,7 +2992,7 @@ bool FUserManagerEOS::QueryExternalIdMappings(const FUniqueNetId& UserId, const
 			FUniqueNetIdEOSPtr EOSID = FUniqueNetIdEOS::EmptyId();
 			if (Result == EOS_EResult::EOS_Success)
 			{
-				EOSID = LocalUserUniqueNetIds[LocalUserNum];
+				EOSID = GetLocalUserChecked(LocalUserNum).UniqueNetId;
 
 				FGetAccountMappingOptions Options;
 				Options.LocalUserId = GetLocalProductUserId(LocalUserNum);
@@ -3009,10 +3014,10 @@ bool FUserManagerEOS::QueryExternalIdMappings(const FUniqueNetId& UserId, const
 			}
 
 			// Mark all queries as complete
-			TArray<FString>& OngoingQueries = IsPlayerQueryExternalMappingsOngoingForLocalUserMap[LocalUserNum];
+			FLocalUserEOS& LocalUser = GetLocalUserChecked(LocalUserNum);
 			for (const FString& StringId : BatchIds)
 			{
-				OngoingQueries.RemoveSwap(StringId, false);
+				LocalUser.OngoingPlayerQueryExternalMappings.RemoveSwap(StringId, false);
 			}
 
 			const bool bWasSuccessful = Result == EOS_EResult::EOS_Success;
diff --git a/Engine/Plugins/Online/OnlineSubsystemEOS/Source/OnlineSubsystemEOS/Private/UserManagerEOS.h b/Engine/Plugins/Online/OnlineSubsystemEOS/Source/OnlineSubsystemEOS/Private/UserManagerEOS.h
--- a/Engine/Plugins/Online/OnlineSubsystemEOS/Source/OnlineSubsystemEOS/Private/UserManagerEOS.h
+++ b/Engine/Plugins/Online/OnlineSubsystemEOS/Source/OnlineSubsystemEOS/Private/UserManagerEOS.h
@@ -83,8 +83,8 @@ class TOnlinePlayerList
 	FUniqueNetIdEOSRef OwningNetId;
 	/** The array of list class entries */
 	TArray<ListClass> ListEntries;
-	/** Indexed by string form of account id for fast look up */
-	TMap<FString, ListClass> NetIdStringToListEntryMap;
+	/** Indexed by unique id for fast look up */
+	TUniqueNetIdMap<ListClass> UniqueIdToListEntryMap;
 
 public:
 	TOnlinePlayerList(int32 InLocalUserNum, FUniqueNetIdEOSRef InOwningNetId)
@@ -98,31 +98,22 @@ public:
 		return ListEntries;
 	}
 
-	void Add(const FString& InNetId, ListClass InListEntry)
+	void Add(const FUniqueNetIdRef& InNetId, ListClass InListEntry)
 	{
 		ListEntries.Add(InListEntry);
-		NetIdStringToListEntryMap.Add(InNetId, InListEntry);
+		UniqueIdToListEntryMap.Add(InNetId, InListEntry);
 	}
 
-	void Remove(const FString& InNetId, ListClass InListEntry)
+	void Remove(const FUniqueNetIdRef& InNetId, ListClass InListEntry)
 	{
-		NetIdStringToListEntryMap.Remove(InNetId);
+		UniqueIdToListEntryMap.Remove(InNetId);
 		ListEntries.Remove(InListEntry);
 	}
 
 	void Empty(int32 Slack = 0)
 	{
 		ListEntries.Empty(Slack);
-		NetIdStringToListEntryMap.Empty(Slack);
-	}
-
-	void UpdateNetIdStr(const FString& PrevNetId, const FString& NewNetId)
-	{
-		if (ListClass* ListEntry = NetIdStringToListEntryMap.Find(PrevNetId))
-		{
-			NetIdStringToListEntryMap.Add(NewNetId, *ListEntry);
-			NetIdStringToListEntryMap.Remove(PrevNetId);
-		}
+		UniqueIdToListEntryMap.Empty(Slack);
 	}
 
 	ListClassReturnType GetByIndex(int32 Index)
@@ -134,9 +125,9 @@ public:
 		return ListClassReturnType();
 	}
 
-	ListClassReturnType GetByNetIdString(const FString& NetId)
+	ListClassReturnType GetByNetId(const FUniqueNetIdRef& NetId)
 	{
-		const ListClass* Found = NetIdStringToListEntryMap.Find(NetId);
+		const ListClass* Found = UniqueIdToListEntryMap.Find(NetId);
 		if (Found != nullptr)
 		{
 			return *Found;
@@ -155,8 +146,8 @@ public:
 	}
 	virtual ~FFriendsListEOS() = default;
 };
-
 typedef TSharedRef<FFriendsListEOS> FFriendsListEOSRef;
+typedef TSharedPtr<FFriendsListEOS> FFriendsListEOSPtr;
 
 struct FNotificationIdCallbackPair
 {
@@ -174,6 +165,53 @@ struct FNotificationIdCallbackPair
 		delete Callback;
 	}
 };
+typedef TSharedPtr<FNotificationIdCallbackPair> FNotificationIdCallbackPairPtr;
+
+/** Cache for the info passed on to ReadFriendsList, kept while user info queries complete */
+struct ReadUserListInfo
+{
+	const int32 LocalUserNum;
+	const FString ListName;
+	const FOnReadFriendsListComplete Delegate;
+
+	ReadUserListInfo(int32 InLocalUserNum, const FString& InListName, FOnReadFriendsListComplete InDelegate)
+		: LocalUserNum(InLocalUserNum), ListName(InListName), Delegate(MoveTemp(InDelegate))
+	{
+	}
+
+	void ExecuteDelegateIfBound(bool bWasSuccessful, const FString& ErrorStr) const
+	{
+		Delegate.ExecuteIfBound(LocalUserNum, bWasSuccessful, ListName, ErrorStr);
+	};
+};
+
+/** Information related to ongoing operations and state for local users */
+struct FLocalUserEOS
+{
+	/** UniqueNetId for the local user, set after successful authentication */
+	FUniqueNetIdEOSPtr UniqueNetId;
+
+	/** Online Account information for the local user */
+	FUserOnlineAccountEOSPtr UserOnlineAccount;
+
+	/** Friends information for the local user */
+	FFriendsListEOSPtr FriendsList;
+
+	/** Epic ids for friends with ongoing queries */
+	TArray<EOS_EpicAccountId> OngoingFriendQueryUserInfo;
+
+	/** External ids still in process of resolution for local user */
+	TArray<FString> OngoingPlayerQueryExternalMappings;
+
+	/** This keeps track of all delegates passed on to calls to ReadFriendsList, so they will all be executed once the process finishes */
+	TArray<ReadUserListInfo> CachedReadUserListInfo;
+
+	/** Last Login Credentials used for a login attempt */
+	TSharedPtr<FOnlineAccountCredentials> LastLoginCredentials;
+
+	/** Struct containing the information related to the connect login notification for the user */
+	FNotificationIdCallbackPairPtr ConnectLoginNotification;
+};
 
 /**
  * EOS service implementation of the online identity, friends, and user info interfaces
@@ -312,28 +350,30 @@ public:
 	 */
 	FUserManagerEOS() = delete;
 
-	bool ConnectLoginEAS(int32 LocalUserNum, EOS_EpicAccountId AccountId);
-	void LoginViaExternalAuth(int32 LocalUserNum);
-	void CreateConnectedLogin(int32 LocalUserNum, EOS_EpicAccountId AccountId, EOS_ContinuanceToken Token);
-	void LinkEAS(int32 LocalUserNum, EOS_ContinuanceToken Token);
+	bool ConnectLoginEAS(int32 LocalUserNum, EOS_EpicAccountId AccountId, const FOnlineAccountCredentials& AccountCredentials);
+	void LoginViaExternalAuth(int32 LocalUserNum, const FOnlineAccountCredentials& AccountCredentials);
+	void CreateConnectedLogin(int32 LocalUserNum, EOS_EpicAccountId AccountId, EOS_ContinuanceToken Token, const FOnlineAccountCredentials& AccountCredentials);
+	void LinkEAS(int32 LocalUserNum, EOS_ContinuanceToken Token, const FOnlineAccountCredentials& AccountCredentials);
 	void RefreshConnectLogin(int32 LocalUserNum);
-	bool ConnectLoginNoEAS(int32 LocalUserNum);
+	bool ConnectLoginNoEAS(int32 LocalUserNum, const FOnlineAccountCredentials& AccountCredentials);
 
-	void FullLoginCallback(int32 LocalUserNum, EOS_EpicAccountId AccountId, EOS_ProductUserId UserId);
+	void FullLoginCallback(int32 LocalUserNum, EOS_EpicAccountId AccountId, EOS_ProductUserId UserId, const FOnlineAccountCredentials& AccountCredentials);
 	void FriendStatusChanged(const EOS_Friends_OnFriendsUpdateInfo* Data);
 	void LoginStatusChanged(const EOS_Auth_LoginStatusChangedCallbackInfo* Data);
 
 	int32 GetDefaultLocalUser() const { return DefaultLocalUser; }
 
 private:
+	FLocalUserEOS& GetLocalUserChecked(int32 LocalUserNum);
+
 	void CallEOSAuthLogin(int32 LocalUserNum, const FOnlineAccountCredentials& Credentials);
 
 	void RemoveLocalUser(int32 LocalUserNum);
-	void AddLocalUser(int32 LocalUserNum, EOS_EpicAccountId EpicAccountId, EOS_ProductUserId UserId);
+	FLocalUserEOS& AddLocalUser(int32 LocalUserNum, EOS_EpicAccountId EpicAccountId, EOS_ProductUserId UserId, const FOnlineAccountCredentials& AccountCredentials);
 
 	FOnlineFriendEOSRef AddFriend(int32 LocalUserNum, EOS_EpicAccountId EpicAccountId);
-	void AddRemotePlayer(int32 LocalUserNum, const FString& NetId, EOS_EpicAccountId EpicAccountId);
-	void AddRemotePlayer(int32 LocalUserNum, const FString& NetId, EOS_EpicAccountId EpicAccountId, FUniqueNetIdEOSPtr UniqueNetId, FOnlineUserPtr OnlineUser, IAttributeAccessInterfaceRef AttributeRef);
+	void AddRemotePlayer(int32 LocalUserNum, EOS_EpicAccountId EpicAccountId);
+	void AddRemotePlayer(int32 LocalUserNum, EOS_EpicAccountId EpicAccountId, IAttributeAccessInterfaceRef AttributeRef);
 	void UpdateRemotePlayerProductUserId(EOS_EpicAccountId AccountId, EOS_ProductUserId UserId);
 	void ReadUserInfo(int32 LocalUserNum, EOS_EpicAccountId EpicAccountId);
 
@@ -342,7 +382,7 @@ private:
 	void ProcessReadFriendsListComplete(int32 LocalUserNum, bool bWasSuccessful, const FString& ErrorStr);
 
 	void UpdatePresence(int32 LocalUserNum, EOS_EpicAccountId AccountId);
-	void UpdateFriendPresence(const FString& FriendId, FOnlineUserPresenceRef Presence);
+	void UpdateFriendPresence(const FUniqueNetIdEOSRef& FriendId, FOnlineUserPresenceRef Presence);
 
 	void GetPlatformAuthToken(int32 LocalUserNum, const FOnGetLinkedAccountAuthTokenCompleteDelegate& Delegate) const;
 	FString GetPlatformDisplayName(int32 LocalUserNum) const;
@@ -354,61 +394,25 @@ private:
 	int32 DefaultLocalUser;
 
 	/** Notification state for SDK events */
-	EOS_NotificationId LoginNotificationId;
-	FCallbackBase* LoginNotificationCallback;
-	EOS_NotificationId FriendsNotificationId;
-	FCallbackBase* FriendsNotificationCallback;
-	EOS_NotificationId PresenceNotificationId;
-	FCallbackBase* PresenceNotificationCallback;
-	TMap<int32, FNotificationIdCallbackPair*> LocalUserNumToConnectLoginNotifcationMap;
-
-	/** Ids mapped to locally registered users */
-	TLocalUserArray<FUniqueNetIdEOSRef> LocalUserUniqueNetIds;
-	/** General map for all registered online users */
-	TUniqueNetIdMap<FUserOnlineAccountEOSRef> UniqueNetIdToUserAccountMap;
-
-	/** Per user friends lists accessible by user num or net id */
-	TMap<int32, FFriendsListEOSRef> LocalUserNumToFriendsListMap;
-	TMap<FString, FFriendsListEOSRef> NetIdStringToFriendsListMap;
-
-	/** Ids mapped to remote users */
-	TMap<FString, IAttributeAccessInterfaceRef> NetIdStringToAttributeAccessMap;
-	TMap<EOS_EpicAccountId, IAttributeAccessInterfaceRef> EpicAccountIdToAttributeAccessMap;
-
-	/** Ids mapped to remote user presence */
-	TMap<FString, FOnlineUserPresenceRef> NetIdStringToOnlineUserPresenceMap;
-
-	/** Id map to keep track of which friends have been processed during async user info queries */
-	TMap<int32, TArray<EOS_EpicAccountId>> IsFriendQueryUserInfoOngoingForLocalUserMap;
-	/** Id map to keep track of which players still need their external id synced */
-	TMap<int32, TArray<FString>> IsPlayerQueryExternalMappingsOngoingForLocalUserMap;
-
-	/** Cache for the info passed on to ReadFriendsList, kept while user info queries complete */
-	struct ReadUserListInfo
-	{
-		const int32 LocalUserNum;
-		const FString ListName;
-		const FOnReadFriendsListComplete Delegate;
-
-		ReadUserListInfo(int32 InLocalUserNum, const FString& InListName, FOnReadFriendsListComplete InDelegate)
-			: LocalUserNum(InLocalUserNum), ListName(InListName), Delegate(MoveTemp(InDelegate))
-		{
-		}
+	EOS_NotificationId LoginNotificationId = EOS_INVALID_NOTIFICATIONID;
+	FCallbackBase* LoginNotificationCallback = nullptr;
+	EOS_NotificationId FriendsNotificationId = EOS_INVALID_NOTIFICATIONID;
+	FCallbackBase* FriendsNotificationCallback = nullptr;
+	EOS_NotificationId PresenceNotificationId = EOS_INVALID_NOTIFICATIONID;
+	FCallbackBase* PresenceNotificationCallback = nullptr;
+	EOS_NotificationId DisplaySettingsUpdatedId = EOS_INVALID_NOTIFICATIONID;
+	FCallbackBase* DisplaySettingsUpdatedCallback = nullptr;
 
-		void ExecuteDelegateIfBound(bool bWasSuccessful, const FString& ErrorStr) const
-		{
-			Delegate.ExecuteIfBound(LocalUserNum, bWasSuccessful, ListName, ErrorStr);
-		};
-	};
+	/** Information related to ongoing operations and state for local users */
+	TLocalUserArray<FLocalUserEOS> LocalUsers;
 
-	TMap<int32, TArray<ReadUserListInfo>> CachedReadUserListInfoForLocalUserMap;
+	// Online user info
 
-	/** Identifier for the external UI notification callback */
-	EOS_NotificationId DisplaySettingsUpdatedId = EOS_INVALID_NOTIFICATIONID;
-	FCallbackBase* DisplaySettingsUpdatedCallback = nullptr;
+	/** Ids mapped to attribute access information for any user */
+	TUniqueNetIdMap<IAttributeAccessInterfaceRef> UniqueNetIdToAttributeAccessMap;
 
-	/** Last Login Credentials used for a login attempt */
-	TMap<int32, TSharedRef<FOnlineAccountCredentials>> LocalUserNumToLastLoginCredentials;
+	/** Ids mapped to user presence information for any user */
+	TUniqueNetIdMap<FOnlineUserPresenceRef> UniqueNetIdToOnlineUserPresenceMap;
 };
 
 typedef TSharedPtr<FUserManagerEOS, ESPMode::ThreadSafe> FUserManagerEOSPtr;
