- Fix incorrect batch sizes
- Fix multiple batches all reading from the start of the source array.

diff --git a/Engine/Plugins/Online/OnlineSubsystemEOS/Source/OnlineSubsystemEOS/Private/UserManagerEOS.cpp b/Engine/Plugins/Online/OnlineSubsystemEOS/Source/OnlineSubsystemEOS/Private/UserManagerEOS.cpp
--- a/Engine/Plugins/Online/OnlineSubsystemEOS/Source/OnlineSubsystemEOS/Private/UserManagerEOS.cpp
+++ b/Engine/Plugins/Online/OnlineSubsystemEOS/Source/OnlineSubsystemEOS/Private/UserManagerEOS.cpp
@@ -2963,23 +2963,24 @@ bool FUserManagerEOS::QueryExternalIdMappings(const FUniqueNetId& UserId, const
 	GetLocalUserChecked(LocalUserNum).OngoingPlayerQueryExternalMappings.Append(ExternalIds);
 
 	const EOS_ProductUserId LocalUserId = EOSID.GetProductUserId();
-	const int32 NumBatches = (ExternalIds.Num() / EOS_CONNECT_QUERYEXTERNALACCOUNTMAPPINGS_MAX_ACCOUNT_IDS) + 1;
-	int32 QueryStart = 0;
+	const uint32 MaxBatchSize = EOS_CONNECT_QUERYEXTERNALACCOUNTMAPPINGS_MAX_ACCOUNT_IDS;
+	const uint32 NumBatches = FMath::DivideAndRoundUp<uint32>(ExternalIds.Num(), MaxBatchSize);
 	// Process queries in batches since there's a max that can be done at once
-	for (int32 BatchCount = 0; BatchCount < NumBatches; BatchCount++)
+	for (uint32 BatchIdx = 0; BatchIdx < NumBatches; BatchIdx++)
 	{
-		const uint32 AmountToProcess = FMath::Min(ExternalIds.Num() - QueryStart, EOS_CONNECT_QUERYEXTERNALACCOUNTMAPPINGS_MAX_ACCOUNT_IDS);
-		TArray<FString> BatchIds;
-		BatchIds.Empty(AmountToProcess);
-		FQueryByStringIdsOptions Options(AmountToProcess, LocalUserId);
+		const uint32 BatchSrcOffset = BatchIdx * MaxBatchSize;
+		const uint32 BatchSize = FMath::Min(ExternalIds.Num() - BatchSrcOffset, MaxBatchSize);
+
 		// Build an options up per batch
-		for (uint32 ProcessedCount = 0; ProcessedCount < AmountToProcess; ProcessedCount++, QueryStart++)
+		FQueryByStringIdsOptions Options(BatchSize, LocalUserId);
+		for (uint32 DestIdx = 0, SrcIdx = BatchSrcOffset; DestIdx < BatchSize; DestIdx++, SrcIdx++)
 		{
-			FCStringAnsi::Strncpy(Options.PointerArray[ProcessedCount], TCHAR_TO_UTF8(*ExternalIds[ProcessedCount]), EOS_CONNECT_EXTERNAL_ACCOUNT_ID_MAX_LENGTH+1);
-			BatchIds.Add(ExternalIds[ProcessedCount]);
+			FCStringAnsi::Strncpy(Options.PointerArray[DestIdx], TCHAR_TO_UTF8(*ExternalIds[SrcIdx]), EOS_CONNECT_EXTERNAL_ACCOUNT_ID_MAX_LENGTH+1);
 		}
+
+		TArray<FString> BatchIds(ExternalIds.GetData() + BatchSrcOffset, BatchSize);
 		FQueryByStringIdsCallback* CallbackObj = new FQueryByStringIdsCallback(AsWeak());
-		CallbackObj->CallbackLambda = [LocalUserNum, QueryOptions, BatchIds, this, Delegate](const EOS_Connect_QueryExternalAccountMappingsCallbackInfo* Data)
+		CallbackObj->CallbackLambda = [LocalUserNum, QueryOptions, BatchIds = MoveTemp(BatchIds), this, Delegate](const EOS_Connect_QueryExternalAccountMappingsCallbackInfo* Data)
 		{
 			EOS_EResult Result = Data->ResultCode;
 			if (GetLoginStatus(LocalUserNum) != ELoginStatus::LoggedIn)
