Fix to Friend information not being updated properly in UserManagerEOS

diff --git a/Engine/Plugins/Online/OnlineSubsystemEOS/Source/OnlineSubsystemEOS/Private/OnlineSubsystemEOSTypes.h b/Engine/Plugins/Online/OnlineSubsystemEOS/Source/OnlineSubsystemEOS/Private/OnlineSubsystemEOSTypes.h
--- a/Engine/Plugins/Online/OnlineSubsystemEOS/Source/OnlineSubsystemEOS/Private/OnlineSubsystemEOSTypes.h
+++ b/Engine/Plugins/Online/OnlineSubsystemEOS/Source/OnlineSubsystemEOS/Private/OnlineSubsystemEOSTypes.h
@@ -127,6 +127,15 @@ private:
 class IAttributeAccessInterface
 {
 public:
+	virtual TMap<FString, FString> GetInternalAttributes() const
+	{
+		return TMap<FString, FString>();
+	}
+
+	virtual void UpdateInternalAttributes(const TMap<FString, FString>& InternalAttributes)
+	{
+	}
+
 	virtual void SetInternalAttribute(const FString& AttrName, const FString& AttrValue)
 	{
 	}
@@ -149,6 +158,8 @@ class TOnlineUserEOS
 	, public AttributeAccessClass
 {
 public:
+	friend class FUserManagerEOS;
+
 	TOnlineUserEOS(const FUniqueNetIdEOSRef& InNetIdRef)
 		: UserIdRef(InNetIdRef)
 	{
@@ -187,6 +198,15 @@ public:
 		return false;
 	}
 //~FOnlineUser
+	virtual TMap<FString, FString> GetInternalAttributes() const override
+	{
+		return UserAttributes;
+	}
+
+	virtual void UpdateInternalAttributes(const TMap<FString, FString>& InternalAttributes) override
+	{
+		UserAttributes.Append(InternalAttributes);
+	}
 
 	virtual void SetInternalAttribute(const FString& AttrName, const FString& AttrValue)
 	{
diff --git a/Engine/Plugins/Online/OnlineSubsystemEOS/Source/OnlineSubsystemEOS/Private/UserManagerEOS.cpp b/Engine/Plugins/Online/OnlineSubsystemEOS/Source/OnlineSubsystemEOS/Private/UserManagerEOS.cpp
--- a/Engine/Plugins/Online/OnlineSubsystemEOS/Source/OnlineSubsystemEOS/Private/UserManagerEOS.cpp
+++ b/Engine/Plugins/Online/OnlineSubsystemEOS/Source/OnlineSubsystemEOS/Private/UserManagerEOS.cpp
@@ -1866,7 +1866,7 @@ FOnlineFriendEOSRef FUserManagerEOS::AddFriend(int32 LocalUserNum, EOS_EpicAccou
 	FriendRef->SetInviteStatus(ToEInviteStatus(Status));
 
 	// Add this friend as a remote player if we haven't already (this will grab user info)
-	if (!FUniqueNetIdEOSRegistry::Find(EpicAccountId))
+	if (!UniqueNetIdToAttributeAccessMap.Contains(FriendNetId))
 	{
 		AddRemotePlayer(LocalUserNum, EpicAccountId, FriendRef);
 	}
@@ -2775,6 +2775,12 @@ void FUserManagerEOS::ReadUserInfo(int32 LocalUserNum, EOS_EpicAccountId EpicAcc
 			const FUniqueNetIdEOSRef EOSId = FUniqueNetIdEOSRegistry::FindChecked(Data->TargetUserId);
 			IAttributeAccessInterfaceRef AttributeAccessRef = UniqueNetIdToAttributeAccessMap[EOSId];
 			UpdateUserInfo(AttributeAccessRef, Data->LocalUserId, Data->TargetUserId);
+
+			FLocalUserEOS& LocalUser = GetLocalUserChecked(LocalUserNum);
+			if (FOnlineFriendEOSPtr FriendPtr = LocalUser.FriendsList->GetByNetId(EOSId))
+			{
+				FriendPtr->UpdateInternalAttributes(AttributeAccessRef->GetInternalAttributes());
+			}
 		}
 
 		// We mark this player as processed
