Refactored OSS EOS Session analytics methods to skip execution and print Verbose log if EAS is not being used.

diff --git a/Engine/Plugins/Online/OnlineSubsystemEOS/Source/OnlineSubsystemEOS/Private/OnlineSessionEOS.cpp b/Engine/Plugins/Online/OnlineSubsystemEOS/Source/OnlineSubsystemEOS/Private/OnlineSessionEOS.cpp
--- a/Engine/Plugins/Online/OnlineSubsystemEOS/Source/OnlineSubsystemEOS/Private/OnlineSessionEOS.cpp
+++ b/Engine/Plugins/Online/OnlineSubsystemEOS/Source/OnlineSubsystemEOS/Private/OnlineSessionEOS.cpp
@@ -1256,9 +1256,10 @@ struct FBeginMetricsOptions :
 
 void FOnlineSessionEOS::BeginSessionAnalytics(FNamedOnlineSession* Session)
 {
-	int32 LocalUserNum = EOSSubsystem->UserManager->GetDefaultLocalUser();
-	FOnlineUserPtr LocalUser = EOSSubsystem->UserManager->GetLocalOnlineUser(LocalUserNum);
-	if (LocalUser.IsValid())
+	const int32 LocalUserNum = EOSSubsystem->UserManager->GetDefaultLocalUser();
+	const FOnlineUserPtr LocalUser = EOSSubsystem->UserManager->GetLocalOnlineUser(LocalUserNum);
+	const EOS_EpicAccountId AccountId = EOSSubsystem->UserManager->GetLocalEpicAccountId(LocalUserNum);
+	if (LocalUser.IsValid() && AccountId != nullptr)
 	{
 		TSharedPtr<const FOnlineSessionInfoEOS> SessionInfoEOS = StaticCastSharedPtr<const FOnlineSessionInfoEOS>(Session->SessionInfo);
 
@@ -1267,7 +1268,7 @@ void FOnlineSessionEOS::BeginSessionAnalytics(FNamedOnlineSession* Session)
 		FString DisplayName = LocalUser->GetDisplayName();
 		FCStringAnsi::Strncpy(Options.DisplayNameAnsi, TCHAR_TO_UTF8(*DisplayName), EOS_OSS_STRING_BUFFER_LENGTH);
 		Options.AccountIdType = EOS_EMetricsAccountIdType::EOS_MAIT_Epic;
-		Options.AccountId.Epic = EOSSubsystem->UserManager->GetLocalEpicAccountId(LocalUserNum);
+		Options.AccountId.Epic = AccountId;
 
 		EOS_EResult Result = EOS_Metrics_BeginPlayerSession(EOSSubsystem->MetricsHandle, &Options);
 		if (Result != EOS_EResult::EOS_Success)
@@ -1275,6 +1276,10 @@ void FOnlineSessionEOS::BeginSessionAnalytics(FNamedOnlineSession* Session)
 			UE_LOG_ONLINE_SESSION(Error, TEXT("EOS_Metrics_BeginPlayerSession() returned EOS result code (%s)"), ANSI_TO_TCHAR(EOS_EResult_ToString(Result)));
 		}
 	}
+	else
+	{
+		UE_LOG_ONLINE_SESSION(Verbose, TEXT("[FOnlineSessionEOS::BeginSessionAnalytics] EOS_Metrics_BeginPlayerSession was not called. Needed AccountId was invalid for LocalUserNum [%d]"), LocalUserNum);
+	}
 }
 
 template<typename BaseStruct>
@@ -1856,13 +1861,13 @@ struct FEndMetricsOptions :
 
 void FOnlineSessionEOS::EndSessionAnalytics()
 {
-	int32 LocalUserNum = EOSSubsystem->UserManager->GetDefaultLocalUser();
-	FOnlineUserPtr LocalUser = EOSSubsystem->UserManager->GetLocalOnlineUser(LocalUserNum);
-	if (LocalUser.IsValid())
+	const int32 LocalUserNum = EOSSubsystem->UserManager->GetDefaultLocalUser();
+	const EOS_EpicAccountId AccountId = EOSSubsystem->UserManager->GetLocalEpicAccountId(LocalUserNum);
+	if (AccountId != nullptr)
 	{
 		FEndMetricsOptions Options;
 		Options.AccountIdType = EOS_EMetricsAccountIdType::EOS_MAIT_Epic;
-		Options.AccountId.Epic = EOSSubsystem->UserManager->GetLocalEpicAccountId(LocalUserNum);
+		Options.AccountId.Epic = AccountId;
 
 		EOS_EResult Result = EOS_Metrics_EndPlayerSession(EOSSubsystem->MetricsHandle, &Options);
 		if (Result != EOS_EResult::EOS_Success)
@@ -1870,6 +1875,10 @@ void FOnlineSessionEOS::EndSessionAnalytics()
 			UE_LOG_ONLINE_SESSION(Error, TEXT("EOS_Metrics_EndPlayerSession() returned EOS result code (%s)"), ANSI_TO_TCHAR(EOS_EResult_ToString(Result)));
 		}
 	}
+	else
+	{
+		UE_LOG_ONLINE_SESSION(Verbose, TEXT("[FOnlineSessionEOS::EndSessionAnalytics] EOS_Metrics_EndPlayerSession was not called. Needed AccountId was invalid for LocalUserNum [%d]"), LocalUserNum);
+	}
 }
 
 struct FSessionDestroyOptions :
