Removal of "_" character from eos presence parameter, as it causes API session-related calls in some platforms to fail.

diff --git a/Engine/Plugins/Online/OnlineSubsystemEOS/Source/OnlineSubsystemEOS/Private/OnlineSessionEOS.cpp b/Engine/Plugins/Online/OnlineSubsystemEOS/Source/OnlineSubsystemEOS/Private/OnlineSessionEOS.cpp
--- a/Engine/Plugins/Online/OnlineSubsystemEOS/Source/OnlineSubsystemEOS/Private/OnlineSessionEOS.cpp
+++ b/Engine/Plugins/Online/OnlineSubsystemEOS/Source/OnlineSubsystemEOS/Private/OnlineSessionEOS.cpp
@@ -25,7 +25,7 @@
 #include "eos_metrics.h"
 #include "eos_lobby.h"
 
-#define USES_PRESENCE_ATTRIBUTE_KEY FName("USES_PRESENCE")
+#define USES_PRESENCE_ATTRIBUTE_KEY FName("USESPRESENCE")
 
 FString MakeStringFromAttributeValue(const EOS_Sessions_AttributeData* Attribute)
 {
