OSS EOS Sessions changes:
- Addition of logic to transmit overlay session join as an OnSessionUserInviteAccepted delegate.
- Addition of support for OnSessionInviteReceived delegate
- Refactoring to resolve all user ids before any logic is processed

diff --git a/Engine/Plugins/Online/OnlineSubsystemEOS/Source/OnlineSubsystemEOS/Private/OnlineSessionEOS.cpp b/Engine/Plugins/Online/OnlineSubsystemEOS/Source/OnlineSubsystemEOS/Private/OnlineSessionEOS.cpp
--- a/Engine/Plugins/Online/OnlineSubsystemEOS/Source/OnlineSubsystemEOS/Private/OnlineSessionEOS.cpp
+++ b/Engine/Plugins/Online/OnlineSubsystemEOS/Source/OnlineSubsystemEOS/Private/OnlineSessionEOS.cpp
@@ -421,6 +421,7 @@ void FOnlineSessionInfoEOS::InitLAN(FOnlineSubsystemEOS* Subsystem)
 
 typedef TEOSGlobalCallback<EOS_Sessions_OnSessionInviteReceivedCallback, EOS_Sessions_SessionInviteReceivedCallbackInfo, FOnlineSessionEOS> FSessionInviteReceivedCallback;
 typedef TEOSGlobalCallback<EOS_Sessions_OnSessionInviteAcceptedCallback, EOS_Sessions_SessionInviteAcceptedCallbackInfo, FOnlineSessionEOS> FSessionInviteAcceptedCallback;
+typedef TEOSGlobalCallback<EOS_Sessions_OnJoinSessionAcceptedCallback, EOS_Sessions_JoinSessionAcceptedCallbackInfo, FOnlineSessionEOS> FJoinSessionAcceptedCallback;
 
 // Lobby session callbacks
 typedef TEOSCallback<EOS_Lobby_OnCreateLobbyCallback, EOS_Lobby_CreateLobbyCallbackInfo, FOnlineSessionEOS> FLobbyCreatedCallback;
@@ -436,6 +437,7 @@ typedef TEOSCallback<EOS_LobbySearch_OnFindCallback, EOS_LobbySearch_FindCallbac
 typedef TEOSGlobalCallback<EOS_Lobby_OnLobbyUpdateReceivedCallback, EOS_Lobby_LobbyUpdateReceivedCallbackInfo, FOnlineSessionEOS> FLobbyUpdateReceivedCallback;
 typedef TEOSGlobalCallback<EOS_Lobby_OnLobbyMemberUpdateReceivedCallback, EOS_Lobby_LobbyMemberUpdateReceivedCallbackInfo, FOnlineSessionEOS> FLobbyMemberUpdateReceivedCallback;
 typedef TEOSGlobalCallback<EOS_Lobby_OnLobbyMemberStatusReceivedCallback, EOS_Lobby_LobbyMemberStatusReceivedCallbackInfo, FOnlineSessionEOS> FLobbyMemberStatusReceivedCallback;
+typedef TEOSGlobalCallback<EOS_Lobby_OnLobbyInviteReceivedCallback, EOS_Lobby_LobbyInviteReceivedCallbackInfo, FOnlineSessionEOS> FLobbyInviteReceivedCallback;
 typedef TEOSGlobalCallback<EOS_Lobby_OnLobbyInviteAcceptedCallback, EOS_Lobby_LobbyInviteAcceptedCallbackInfo, FOnlineSessionEOS> FLobbyInviteAcceptedCallback;
 typedef TEOSGlobalCallback<EOS_Lobby_OnJoinLobbyAcceptedCallback, EOS_Lobby_JoinLobbyAcceptedCallbackInfo, FOnlineSessionEOS> FJoinLobbyAcceptedCallback;
 
@@ -459,42 +461,7 @@ FOnlineSessionEOS::~FOnlineSessionEOS()
 
 void FOnlineSessionEOS::Init()
 {
-	// Register for session invite notifications
-	FSessionInviteAcceptedCallback* SessionInviteAcceptedCallbackObj = new FSessionInviteAcceptedCallback(FOnlineSessionEOSWeakPtr(AsShared()));
-	SessionInviteAcceptedCallback = SessionInviteAcceptedCallbackObj;
-	SessionInviteAcceptedCallbackObj->CallbackLambda = [this](const EOS_Sessions_SessionInviteAcceptedCallbackInfo* Data)
-	{
-		FUniqueNetIdEOSPtr NetId = EOSSubsystem->UserManager->GetLocalUniqueNetIdEOS(Data->LocalUserId);
-		if (!NetId.IsValid())
-		{
-			UE_LOG_ONLINE_SESSION(Warning, TEXT("Cannot accept invite due to unknown user (%s)"), *LexToString(Data->LocalUserId));
-			TriggerOnSessionUserInviteAcceptedDelegates(false, 0, NetId, FOnlineSessionSearchResult());
-			return;
-		}
-		int32 LocalUserNum = EOSSubsystem->UserManager->GetLocalUserNumFromUniqueNetId(*NetId);
-
-		EOS_Sessions_CopySessionHandleByInviteIdOptions Options = { };
-		Options.ApiVersion = 1;
-		UE_EOS_CHECK_API_MISMATCH(EOS_SESSIONS_COPYSESSIONHANDLEBYINVITEID_API_LATEST, 1);
-		Options.InviteId = Data->InviteId;
-		EOS_HSessionDetails SessionDetails = nullptr;
-		EOS_EResult Result = EOS_Sessions_CopySessionHandleByInviteId(EOSSubsystem->SessionsHandle, &Options, &SessionDetails);
-		if (Result == EOS_EResult::EOS_Success)
-		{
-			LastInviteSearch = MakeShared<FOnlineSessionSearch>();
-			AddSearchResult(SessionDetails, LastInviteSearch.ToSharedRef());
-			TriggerOnSessionUserInviteAcceptedDelegates(true, LocalUserNum, NetId, LastInviteSearch->SearchResults[0]);
-		}
-		else
-		{
-			UE_LOG_ONLINE_SESSION(Warning, TEXT("EOS_Sessions_CopySessionHandleByInviteId not successful. Finished with EOS_EResult %s"), ANSI_TO_TCHAR(EOS_EResult_ToString(Result)));
-			TriggerOnSessionUserInviteAcceptedDelegates(false, LocalUserNum, NetId, FOnlineSessionSearchResult());
-		}
-	};
-	EOS_Sessions_AddNotifySessionInviteAcceptedOptions Options = { };
-	Options.ApiVersion = 1;
-	UE_EOS_CHECK_API_MISMATCH(EOS_SESSIONS_ADDNOTIFYSESSIONINVITEACCEPTED_API_LATEST, 1);
-	SessionInviteAcceptedId = EOS_Sessions_AddNotifySessionInviteAccepted(EOSSubsystem->SessionsHandle, &Options, SessionInviteAcceptedCallbackObj, SessionInviteAcceptedCallbackObj->GetCallbackPtr());
+	RegisterSessionNotifications();
 
 	// Lobbies
 	LobbyHandle = EOS_Platform_GetLobbyInterface(*EOSSubsystem->EOSPlatformHandle);
@@ -635,6 +602,142 @@ int32 FOnlineSessionEOS::GetDefaultLocalUserForLobby(const FUniqueNetIdString& S
 	return INVALID_LOCAL_USER;
 }
 
+void FOnlineSessionEOS::RegisterSessionNotifications()
+{
+	// Register for session invite received notifications
+	FSessionInviteReceivedCallback* SessionInviteReceivedCallbackObj = new FSessionInviteReceivedCallback(FOnlineSessionEOSWeakPtr(AsShared()));
+	SessionInviteReceivedCallback = SessionInviteReceivedCallbackObj;
+	SessionInviteReceivedCallbackObj->CallbackLambda = [this](const EOS_Sessions_SessionInviteReceivedCallbackInfo* Data)
+	{
+		OnSessionInviteReceived(Data);
+	};
+	EOS_Sessions_AddNotifySessionInviteReceivedOptions AddNotifySessionInviteReceivedOptions = { };
+	AddNotifySessionInviteReceivedOptions.ApiVersion = 1;
+	UE_EOS_CHECK_API_MISMATCH(EOS_SESSIONS_ADDNOTIFYSESSIONINVITERECEIVED_API_LATEST, 1);
+	SessionInviteReceivedId = EOS_Sessions_AddNotifySessionInviteReceived(EOSSubsystem->SessionsHandle, &AddNotifySessionInviteReceivedOptions, SessionInviteReceivedCallbackObj, SessionInviteReceivedCallbackObj->GetCallbackPtr());
+
+	// Register for session invite accepted notifications
+	FSessionInviteAcceptedCallback* SessionInviteAcceptedCallbackObj = new FSessionInviteAcceptedCallback(FOnlineSessionEOSWeakPtr(AsShared()));
+	SessionInviteAcceptedCallback = SessionInviteAcceptedCallbackObj;
+	SessionInviteAcceptedCallbackObj->CallbackLambda = [this](const EOS_Sessions_SessionInviteAcceptedCallbackInfo* Data)
+	{
+		OnSessionInviteAccepted(Data);
+	};
+	EOS_Sessions_AddNotifySessionInviteAcceptedOptions AddNotifySessionInviteAcceptedOptions = { };
+	AddNotifySessionInviteAcceptedOptions.ApiVersion = 1;
+	UE_EOS_CHECK_API_MISMATCH(EOS_SESSIONS_ADDNOTIFYSESSIONINVITEACCEPTED_API_LATEST, 1);
+	SessionInviteAcceptedId = EOS_Sessions_AddNotifySessionInviteAccepted(EOSSubsystem->SessionsHandle, &AddNotifySessionInviteAcceptedOptions, SessionInviteAcceptedCallbackObj, SessionInviteAcceptedCallbackObj->GetCallbackPtr());
+
+	// Register for join session accepted notifications
+	FJoinSessionAcceptedCallback* JoinSessionAcceptedCallbackObj = new FJoinSessionAcceptedCallback(FOnlineSessionEOSWeakPtr(AsShared()));
+	JoinSessionAcceptedCallback = JoinSessionAcceptedCallbackObj;
+	JoinSessionAcceptedCallbackObj->CallbackLambda = [this](const EOS_Sessions_JoinSessionAcceptedCallbackInfo* Data)
+	{
+		OnJoinSessionAccepted(Data);
+	};
+	EOS_Sessions_AddNotifyJoinSessionAcceptedOptions AddNotifyJoinSessionAcceptedOptions = { };
+	AddNotifyJoinSessionAcceptedOptions.ApiVersion = 1;
+	UE_EOS_CHECK_API_MISMATCH(EOS_SESSIONS_ADDNOTIFYJOINSESSIONACCEPTED_API_LATEST, 1);
+	SessionInviteAcceptedId = EOS_Sessions_AddNotifyJoinSessionAccepted(EOSSubsystem->SessionsHandle, &AddNotifyJoinSessionAcceptedOptions, JoinSessionAcceptedCallbackObj, JoinSessionAcceptedCallbackObj->GetCallbackPtr());
+}
+
+
+void FOnlineSessionEOS::OnSessionInviteReceived(const EOS_Sessions_SessionInviteReceivedCallbackInfo* Data)
+{
+	EOSSubsystem->UserManager->ResolveUniqueNetIds(EOSSubsystem->UserManager->GetDefaultLocalUser(), { Data->LocalUserId, Data->TargetUserId }, [this, LocalUserId = Data->LocalUserId, TargetUserId = Data->TargetUserId, InviteId = FString(UTF8_TO_TCHAR(Data->InviteId))](TMap<EOS_ProductUserId, FUniqueNetIdEOSRef> ResolvedUniqueNetIds)
+		{
+			if (!ResolvedUniqueNetIds.Contains(LocalUserId))
+			{
+				// We'll print a warning but not trigger the delegate as we have no information to transmit with it
+				UE_LOG_ONLINE_SESSION(Warning, TEXT("Cannot process invite due to unknown user (%s)"), *LexToString(LocalUserId));
+				return;
+			}
+
+			if (!ResolvedUniqueNetIds.Contains(TargetUserId))
+			{
+				// We'll print a warning but not trigger the delegate as we have no information to transmit with it
+				UE_LOG_ONLINE_SESSION(Warning, TEXT("Cannot process invite due to unknown user (%s)"), *LexToString(TargetUserId));
+				return;
+			}
+
+			const FUniqueNetIdEOSRef NetId = ResolvedUniqueNetIds[LocalUserId];
+			const FUniqueNetIdEOSRef FromNetId = ResolvedUniqueNetIds[TargetUserId];
+			const int32 LocalUserNum = EOSSubsystem->UserManager->GetLocalUserNumFromUniqueNetId(*NetId);
+
+			EOS_Sessions_CopySessionHandleByInviteIdOptions Options = { };
+			Options.ApiVersion = 1;
+			UE_EOS_CHECK_API_MISMATCH(EOS_SESSIONS_COPYSESSIONHANDLEBYINVITEID_API_LATEST, 1);
+			const auto InviteIdUTF8 = StringCast<UTF8CHAR>(*InviteId);
+			Options.InviteId = (const char*)InviteIdUTF8.Get();
+			EOS_HSessionDetails SessionDetails = nullptr;
+			EOS_EResult Result = EOS_Sessions_CopySessionHandleByInviteId(EOSSubsystem->SessionsHandle, &Options, &SessionDetails);
+			if (Result == EOS_EResult::EOS_Success)
+			{
+				LastInviteSearch = MakeShared<FOnlineSessionSearch>();
+				AddSearchResult(SessionDetails, LastInviteSearch.ToSharedRef());
+				TriggerOnSessionInviteReceivedDelegates(*NetId, *FromNetId, EOSSubsystem->GetAppId(), LastInviteSearch->SearchResults[0]);
+			}
+			else
+			{
+				UE_LOG_ONLINE_SESSION(Warning, TEXT("EOS_Sessions_CopySessionHandleByInviteId not successful. Finished with EOS_EResult %s"), ANSI_TO_TCHAR(EOS_EResult_ToString(Result)));
+				TriggerOnSessionInviteReceivedDelegates(*NetId, *FromNetId, EOSSubsystem->GetAppId(), FOnlineSessionSearchResult());
+			}
+		});
+}
+
+void FOnlineSessionEOS::OnSessionInviteAccepted(const EOS_Sessions_SessionInviteAcceptedCallbackInfo* Data)
+{
+	EOSSubsystem->UserManager->ResolveUniqueNetId(EOSSubsystem->UserManager->GetDefaultLocalUser(), Data->LocalUserId, [this, InviteId = FString(UTF8_TO_TCHAR(Data->InviteId))](FUniqueNetIdEOSRef ResolvedUniqueNetId)
+		{
+			const int32 LocalUserNum = EOSSubsystem->UserManager->GetLocalUserNumFromUniqueNetId(*ResolvedUniqueNetId);
+
+			EOS_Sessions_CopySessionHandleByInviteIdOptions Options = { };
+			Options.ApiVersion = 1;
+			UE_EOS_CHECK_API_MISMATCH(EOS_SESSIONS_COPYSESSIONHANDLEBYINVITEID_API_LATEST, 1);
+			const auto InviteIdUTF8 = StringCast<UTF8CHAR>(*InviteId);
+			Options.InviteId = (const char*)InviteIdUTF8.Get();
+			EOS_HSessionDetails SessionDetails = nullptr;
+			EOS_EResult Result = EOS_Sessions_CopySessionHandleByInviteId(EOSSubsystem->SessionsHandle, &Options, &SessionDetails);
+			if (Result == EOS_EResult::EOS_Success)
+			{
+				LastInviteSearch = MakeShared<FOnlineSessionSearch>();
+				AddSearchResult(SessionDetails, LastInviteSearch.ToSharedRef());
+				TriggerOnSessionUserInviteAcceptedDelegates(true, LocalUserNum, ResolvedUniqueNetId, LastInviteSearch->SearchResults[0]);
+			}
+			else
+			{
+				UE_LOG_ONLINE_SESSION(Warning, TEXT("EOS_Sessions_CopySessionHandleByInviteId not successful. Finished with EOS_EResult %s"), ANSI_TO_TCHAR(EOS_EResult_ToString(Result)));
+				TriggerOnSessionUserInviteAcceptedDelegates(false, LocalUserNum, ResolvedUniqueNetId, FOnlineSessionSearchResult());
+			}
+		});
+}
+
+void FOnlineSessionEOS::OnJoinSessionAccepted(const EOS_Sessions_JoinSessionAcceptedCallbackInfo* Data)
+{
+	EOSSubsystem->UserManager->ResolveUniqueNetId(EOSSubsystem->UserManager->GetDefaultLocalUser(), Data->LocalUserId, [this, UiEventId = Data->UiEventId](FUniqueNetIdEOSRef ResolvedUniqueNetId)
+		{
+			const int32 LocalUserNum = EOSSubsystem->UserManager->GetLocalUserNumFromUniqueNetId(*ResolvedUniqueNetId);
+
+			EOS_Sessions_CopySessionHandleByUiEventIdOptions Options = { };
+			Options.ApiVersion = 1;
+			UE_EOS_CHECK_API_MISMATCH(EOS_SESSIONS_COPYSESSIONHANDLEBYUIEVENTID_API_LATEST, 1);
+			Options.UiEventId = UiEventId;
+			EOS_HSessionDetails SessionDetails = nullptr;
+			EOS_EResult Result = EOS_Sessions_CopySessionHandleByUiEventId(EOSSubsystem->SessionsHandle, &Options, &SessionDetails);
+			if (Result == EOS_EResult::EOS_Success)
+			{
+				LastInviteSearch = MakeShared<FOnlineSessionSearch>();
+				AddSearchResult(SessionDetails, LastInviteSearch.ToSharedRef());
+				TriggerOnSessionUserInviteAcceptedDelegates(true, LocalUserNum, ResolvedUniqueNetId, LastInviteSearch->SearchResults[0]);
+			}
+			else
+			{
+				UE_LOG_ONLINE_SESSION(Warning, TEXT("EOS_Sessions_CopySessionHandleByUiEventId not successful. Finished with EOS_EResult %s"), ANSI_TO_TCHAR(EOS_EResult_ToString(Result)));
+				TriggerOnSessionUserInviteAcceptedDelegates(false, LocalUserNum, ResolvedUniqueNetId, FOnlineSessionSearchResult());
+			}
+		});
+}
+
 void FOnlineSessionEOS::RegisterLobbyNotifications()
 {
 	// Lobby data updates
@@ -679,7 +782,19 @@ void FOnlineSessionEOS::RegisterLobbyNotifications()
 
 	LobbyMemberStatusReceivedId = EOS_Lobby_AddNotifyLobbyMemberStatusReceived(LobbyHandle, &AddNotifyLobbyMemberStatusReceivedOptions, LobbyMemberStatusReceivedCallbackObj, LobbyMemberStatusReceivedCallbackObj->GetCallbackPtr());
 
-	// LobbyInviteReceived we can skip, since it will pop up as system UI
+	// Received lobby invite notifications
+	EOS_Lobby_AddNotifyLobbyInviteReceivedOptions AddNotifyLobbyInviteReceivedOptions = { 0 };
+	AddNotifyLobbyInviteReceivedOptions.ApiVersion = 1;
+	UE_EOS_CHECK_API_MISMATCH(EOS_LOBBY_ADDNOTIFYLOBBYINVITERECEIVED_API_LATEST, 1);
+
+	FLobbyInviteReceivedCallback* LobbyInviteReceivedCallbackObj = new FLobbyInviteReceivedCallback(FOnlineSessionEOSWeakPtr(AsShared()));
+	LobbyInviteReceivedCallback = LobbyInviteReceivedCallbackObj;
+	LobbyInviteReceivedCallbackObj->CallbackLambda = [this](const EOS_Lobby_LobbyInviteReceivedCallbackInfo* Data)
+	{
+		OnLobbyInviteReceived(Data);
+	};
+
+	LobbyInviteReceivedId = EOS_Lobby_AddNotifyLobbyInviteReceived(LobbyHandle, &AddNotifyLobbyInviteReceivedOptions, LobbyInviteReceivedCallbackObj, LobbyInviteReceivedCallbackObj->GetCallbackPtr());
 
 	// Accepted lobby invite notifications
 	EOS_Lobby_AddNotifyLobbyInviteAcceptedOptions AddNotifyLobbyInviteAcceptedOptions = { 0 };
@@ -690,7 +805,7 @@ void FOnlineSessionEOS::RegisterLobbyNotifications()
 	LobbyInviteAcceptedCallback = LobbyInviteAcceptedCallbackObj;
 	LobbyInviteAcceptedCallbackObj->CallbackLambda = [this](const EOS_Lobby_LobbyInviteAcceptedCallbackInfo* Data)
 	{
-		OnLobbyInviteAccepted(Data->InviteId, Data->LocalUserId, Data->TargetUserId);
+		OnLobbyInviteAccepted(Data);
 	};
 
 	LobbyInviteAcceptedId = EOS_Lobby_AddNotifyLobbyInviteAccepted(LobbyHandle, &AddNotifyLobbyInviteAcceptedOptions, LobbyInviteAcceptedCallbackObj, LobbyInviteAcceptedCallbackObj->GetCallbackPtr());
@@ -704,7 +819,7 @@ void FOnlineSessionEOS::RegisterLobbyNotifications()
 	JoinLobbyAcceptedCallback = JoinLobbyAcceptedCallbackObj;
 	JoinLobbyAcceptedCallbackObj->CallbackLambda = [this](const EOS_Lobby_JoinLobbyAcceptedCallbackInfo* Data)
 	{
-		OnJoinLobbyAccepted(Data->LocalUserId, Data->UiEventId);
+		OnJoinLobbyAccepted(Data);
 	};
 
 	JoinLobbyAcceptedId = EOS_Lobby_AddNotifyJoinLobbyAccepted(LobbyHandle, &AddNotifyJoinLobbyAcceptedOptions, JoinLobbyAcceptedCallbackObj, JoinLobbyAcceptedCallbackObj->GetCallbackPtr());
@@ -906,77 +1021,118 @@ void FOnlineSessionEOS::OnMemberStatusReceived(const EOS_LobbyId& LobbyId, const
 	}
 }
 
-void FOnlineSessionEOS::OnLobbyInviteAccepted(const char* InviteId, const EOS_ProductUserId& LocalUserId, const EOS_ProductUserId& TargetUserId)
+void FOnlineSessionEOS::OnLobbyInviteReceived(const EOS_Lobby_LobbyInviteReceivedCallbackInfo* Data)
 {
-	FUniqueNetIdEOSPtr NetId = EOSSubsystem->UserManager->GetLocalUniqueNetIdEOS(LocalUserId);
-	if (!NetId.IsValid())
-	{
-		UE_LOG_ONLINE_SESSION(Warning, TEXT("[FOnlineSessionEOS::OnLobbyInviteAccepted] Cannot accept lobby invite due to unknown user (%s)"), *LexToString(LocalUserId));
-		TriggerOnSessionUserInviteAcceptedDelegates(false, 0, NetId, FOnlineSessionSearchResult());
-		return;
-	}
-	int32 LocalUserNum = EOSSubsystem->UserManager->GetLocalUserNumFromUniqueNetId(*NetId);
+	const TArray<EOS_ProductUserId> PUIdS = { Data->LocalUserId, Data->TargetUserId };
+	EOSSubsystem->UserManager->ResolveUniqueNetIds(EOSSubsystem->UserManager->GetDefaultLocalUser(), PUIdS, [this, LocalUserId = Data->LocalUserId, TargetUserId = Data->TargetUserId, InviteId = FString(UTF8_TO_TCHAR(Data->InviteId))](TMap<EOS_ProductUserId, FUniqueNetIdEOSRef> ResolvedUniqueNetIds)
+		{
+			if (!ResolvedUniqueNetIds.Contains(LocalUserId))
+			{
+				// We'll print a warning but not trigger the delegate as we have no information to transmit with it
+				UE_LOG_ONLINE_SESSION(Warning, TEXT("Cannot process invite due to unknown user (%s)"), *LexToString(LocalUserId));
+				return;
+			}
 
-	EOS_Lobby_CopyLobbyDetailsHandleByInviteIdOptions Options = { };
-	Options.ApiVersion = 1;
-	UE_EOS_CHECK_API_MISMATCH(EOS_LOBBY_COPYLOBBYDETAILSHANDLEBYINVITEID_API_LATEST, 1);
-	Options.InviteId = InviteId;
+			if (!ResolvedUniqueNetIds.Contains(TargetUserId))
+			{
+				// We'll print a warning but not trigger the delegate as we have no information to transmit with it
+				UE_LOG_ONLINE_SESSION(Warning, TEXT("Cannot process invite due to unknown user (%s)"), *LexToString(TargetUserId));
+				return;
+			}
 
-	EOS_HLobbyDetails LobbyDetailsHandle;
+			const FUniqueNetIdEOSRef NetId = ResolvedUniqueNetIds[LocalUserId];
+			const FUniqueNetIdEOSRef FromNetId = ResolvedUniqueNetIds[TargetUserId];
+			const int32 LocalUserNum = EOSSubsystem->UserManager->GetLocalUserNumFromUniqueNetId(*NetId);
 
-	EOS_EResult Result = EOS_Lobby_CopyLobbyDetailsHandleByInviteId(LobbyHandle, &Options, &LobbyDetailsHandle);
-	if (Result == EOS_EResult::EOS_Success)
-	{
-		TSharedRef<FLobbyDetailsEOS> LobbyDetails = MakeShared<FLobbyDetailsEOS>(LobbyDetailsHandle);
+			EOS_Lobby_CopyLobbyDetailsHandleByInviteIdOptions Options = { };
+			Options.ApiVersion = 1;
+			UE_EOS_CHECK_API_MISMATCH(EOS_LOBBY_COPYLOBBYDETAILSHANDLEBYINVITEID_API_LATEST, 1);
+			const auto InviteIdUTF8 = StringCast<UTF8CHAR>(*InviteId);
+			Options.InviteId = (const char*)InviteIdUTF8.Get();
+			EOS_HLobbyDetails LobbyDetailsHandle = nullptr;
+			EOS_EResult Result = EOS_Lobby_CopyLobbyDetailsHandleByInviteId(LobbyHandle, &Options, &LobbyDetailsHandle);
+			if (Result == EOS_EResult::EOS_Success)
+			{
+				TSharedRef<FLobbyDetailsEOS> LobbyDetails = MakeShared<FLobbyDetailsEOS>(LobbyDetailsHandle);
 
-		LastInviteSearch = MakeShared<FOnlineSessionSearch>();
-		AddLobbySearchResult(LobbyDetails, LastInviteSearch.ToSharedRef(), [this, LocalUserNum, NetId](bool bWasSuccessful)
-		{
-			// If we fail to copy the lobby data, we won't add a new search result, so we'll return an empty one
-			TriggerOnSessionUserInviteAcceptedDelegates(bWasSuccessful, LocalUserNum, NetId, bWasSuccessful ? LastInviteSearch->SearchResults.Last() : FOnlineSessionSearchResult());
+				LastInviteSearch = MakeShared<FOnlineSessionSearch>();
+				AddLobbySearchResult(LobbyDetails, LastInviteSearch.ToSharedRef(), [this, LocalUserNum, NetId, FromNetId](bool bWasSuccessful)
+					{
+						TriggerOnSessionInviteReceivedDelegates(*NetId, *FromNetId, EOSSubsystem->GetAppId(), bWasSuccessful ? LastInviteSearch->SearchResults.Last() : FOnlineSessionSearchResult());
+					});
+			}
+			else
+			{
+				UE_LOG_ONLINE_SESSION(Warning, TEXT("EOS_Lobby_CopyLobbyDetailsHandleByInviteId not successful. Finished with EOS_EResult %s"), ANSI_TO_TCHAR(EOS_EResult_ToString(Result)));
+				TriggerOnSessionInviteReceivedDelegates(*NetId, *FromNetId, EOSSubsystem->GetAppId(), FOnlineSessionSearchResult());
+			}
 		});
-	}
-	else
-	{
-		UE_LOG_ONLINE_SESSION(Warning, TEXT("[FOnlineSessionEOS::OnLobbyInviteAccepted] EOS_Lobby_CopyLobbyDetailsHandleByInviteId failed with EOS result code (%s)"), ANSI_TO_TCHAR(EOS_EResult_ToString(Result)));
-		TriggerOnSessionUserInviteAcceptedDelegates(false, LocalUserNum, NetId, FOnlineSessionSearchResult());
-	}
 }
 
-void FOnlineSessionEOS::OnJoinLobbyAccepted(const EOS_ProductUserId& LocalUserId, const EOS_UI_EventId& UiEventId)
+void FOnlineSessionEOS::OnLobbyInviteAccepted(const EOS_Lobby_LobbyInviteAcceptedCallbackInfo* Data)
 {
-	FUniqueNetIdEOSPtr NetId = EOSSubsystem->UserManager->GetLocalUniqueNetIdEOS(LocalUserId);
-	if (!NetId.IsValid())
-	{
-		UE_LOG_ONLINE_SESSION(Warning, TEXT("[FOnlineSessionEOS::OnJoinLobbyAccepted] Cannot join lobby due to unknown user (%s)"), *LexToString(LocalUserId));
-		TriggerOnSessionUserInviteAcceptedDelegates(false, 0, NetId, FOnlineSessionSearchResult());
-		return;
-	}
-	int32 LocalUserNum = EOSSubsystem->UserManager->GetLocalUserNumFromUniqueNetId(*NetId);
+	EOSSubsystem->UserManager->ResolveUniqueNetId(EOSSubsystem->UserManager->GetDefaultLocalUser(), Data->LocalUserId, [this, InviteId = FString(UTF8_TO_TCHAR(Data->InviteId))](FUniqueNetIdEOSRef ResolvedUniqueNetId)
+		{
+			const int32 LocalUserNum = EOSSubsystem->UserManager->GetLocalUserNumFromUniqueNetId(*ResolvedUniqueNetId);
 
-	EOS_Lobby_CopyLobbyDetailsHandleByUiEventIdOptions Options = { 0 };
-	Options.ApiVersion = 1;
-	UE_EOS_CHECK_API_MISMATCH(EOS_LOBBY_COPYLOBBYDETAILSHANDLEBYUIEVENTID_API_LATEST, 1);
-	Options.UiEventId = UiEventId;
+			EOS_Lobby_CopyLobbyDetailsHandleByInviteIdOptions Options = { };
+			Options.ApiVersion = 1;
+			UE_EOS_CHECK_API_MISMATCH(EOS_LOBBY_COPYLOBBYDETAILSHANDLEBYINVITEID_API_LATEST, 1);
+			const auto InviteIdUTF8 = StringCast<UTF8CHAR>(*InviteId);
+			Options.InviteId = (const char*)InviteIdUTF8.Get();
 
-	EOS_HLobbyDetails LobbyDetailsHandle;
-	EOS_EResult Result = EOS_Lobby_CopyLobbyDetailsHandleByUiEventId(LobbyHandle, &Options, &LobbyDetailsHandle);
-	if (Result == EOS_EResult::EOS_Success)
-	{
-		TSharedRef<FLobbyDetailsEOS> LobbyDetails = MakeShared<FLobbyDetailsEOS>(LobbyDetailsHandle);
+			EOS_HLobbyDetails LobbyDetailsHandle;
+
+			EOS_EResult Result = EOS_Lobby_CopyLobbyDetailsHandleByInviteId(LobbyHandle, &Options, &LobbyDetailsHandle);
+			if (Result == EOS_EResult::EOS_Success)
+			{
+				TSharedRef<FLobbyDetailsEOS> LobbyDetails = MakeShared<FLobbyDetailsEOS>(LobbyDetailsHandle);
+
+				LastInviteSearch = MakeShared<FOnlineSessionSearch>();
+				AddLobbySearchResult(LobbyDetails, LastInviteSearch.ToSharedRef(), [this, LocalUserNum, ResolvedUniqueNetId](bool bWasSuccessful)
+					{
+						// If we fail to copy the lobby data, we won't add a new search result, so we'll return an empty one
+						TriggerOnSessionUserInviteAcceptedDelegates(bWasSuccessful, LocalUserNum, ResolvedUniqueNetId, bWasSuccessful ? LastInviteSearch->SearchResults.Last() : FOnlineSessionSearchResult());
+					});
+			}
+			else
+			{
+				UE_LOG_ONLINE_SESSION(Warning, TEXT("[FOnlineSessionEOS::OnLobbyInviteAccepted] EOS_Lobby_CopyLobbyDetailsHandleByInviteId failed with EOS result code (%s)"), ANSI_TO_TCHAR(EOS_EResult_ToString(Result)));
+				TriggerOnSessionUserInviteAcceptedDelegates(false, LocalUserNum, ResolvedUniqueNetId, FOnlineSessionSearchResult());
+			}
+		});
+}
 
-		LastInviteSearch = MakeShared<FOnlineSessionSearch>();
-		AddLobbySearchResult(LobbyDetails, LastInviteSearch.ToSharedRef(), [this, LocalUserNum, NetId](bool bWasSuccessful)
+void FOnlineSessionEOS::OnJoinLobbyAccepted(const EOS_Lobby_JoinLobbyAcceptedCallbackInfo* Data)
+{
+	EOSSubsystem->UserManager->ResolveUniqueNetId(EOSSubsystem->UserManager->GetDefaultLocalUser(), Data->LocalUserId, [this, UiEventId = Data->UiEventId](FUniqueNetIdEOSRef ResolvedUniqueNetId)
 		{
-			// If we fail to copy the lobby data, we won't add a new search result, so we'll return an empty one
-			TriggerOnSessionUserInviteAcceptedDelegates(bWasSuccessful, LocalUserNum, NetId, bWasSuccessful ? LastInviteSearch->SearchResults.Last() : FOnlineSessionSearchResult());
+			const int32 LocalUserNum = EOSSubsystem->UserManager->GetLocalUserNumFromUniqueNetId(*ResolvedUniqueNetId);
+
+			EOS_Lobby_CopyLobbyDetailsHandleByUiEventIdOptions Options = { 0 };
+			Options.ApiVersion = 1;
+			UE_EOS_CHECK_API_MISMATCH(EOS_LOBBY_COPYLOBBYDETAILSHANDLEBYUIEVENTID_API_LATEST, 1);
+			Options.UiEventId = UiEventId;
+
+			EOS_HLobbyDetails LobbyDetailsHandle;
+			EOS_EResult Result = EOS_Lobby_CopyLobbyDetailsHandleByUiEventId(LobbyHandle, &Options, &LobbyDetailsHandle);
+			if (Result == EOS_EResult::EOS_Success)
+			{
+				TSharedRef<FLobbyDetailsEOS> LobbyDetails = MakeShared<FLobbyDetailsEOS>(LobbyDetailsHandle);
+
+				LastInviteSearch = MakeShared<FOnlineSessionSearch>();
+				AddLobbySearchResult(LobbyDetails, LastInviteSearch.ToSharedRef(), [this, LocalUserNum, ResolvedUniqueNetId](bool bWasSuccessful)
+					{
+						// If we fail to copy the lobby data, we won't add a new search result, so we'll return an empty one
+						TriggerOnSessionUserInviteAcceptedDelegates(bWasSuccessful, LocalUserNum, ResolvedUniqueNetId, bWasSuccessful ? LastInviteSearch->SearchResults.Last() : FOnlineSessionSearchResult());
+					});
+			}
+			else
+			{
+				UE_LOG_ONLINE_SESSION(Warning, TEXT("[FOnlineSessionEOS::OnJoinLobbyAccepted] EOS_Lobby_CopyLobbyDetailsHandleByUiEventId failed with EOS result code (%s)"), ANSI_TO_TCHAR(EOS_EResult_ToString(Result)));
+				TriggerOnSessionUserInviteAcceptedDelegates(false, LocalUserNum, ResolvedUniqueNetId, FOnlineSessionSearchResult());
+			}
 		});
-	}
-	else
-	{
-		UE_LOG_ONLINE_SESSION(Warning, TEXT("[FOnlineSessionEOS::OnJoinLobbyAccepted] EOS_Lobby_CopyLobbyDetailsHandleByUiEventId failed with EOS result code (%s)"), ANSI_TO_TCHAR(EOS_EResult_ToString(Result)));
-		TriggerOnSessionUserInviteAcceptedDelegates(false, LocalUserNum, NetId, FOnlineSessionSearchResult());
-	}
 }
 
 bool FOnlineSessionEOS::CreateSession(int32 HostingPlayerNum, FName SessionName, const FOnlineSessionSettings& NewSessionSettings)
diff --git a/Engine/Plugins/Online/OnlineSubsystemEOS/Source/OnlineSubsystemEOS/Private/OnlineSessionEOS.h b/Engine/Plugins/Online/OnlineSubsystemEOS/Source/OnlineSubsystemEOS/Private/OnlineSessionEOS.h
--- a/Engine/Plugins/Online/OnlineSubsystemEOS/Source/OnlineSubsystemEOS/Private/OnlineSessionEOS.h
+++ b/Engine/Plugins/Online/OnlineSubsystemEOS/Source/OnlineSubsystemEOS/Private/OnlineSessionEOS.h
@@ -242,6 +242,8 @@ private:
 	FCallbackBase* LobbyMemberUpdateReceivedCallback;
 	EOS_NotificationId LobbyMemberStatusReceivedId;
 	FCallbackBase* LobbyMemberStatusReceivedCallback;
+	EOS_NotificationId LobbyInviteReceivedId;
+	FCallbackBase* LobbyInviteReceivedCallback;
 	EOS_NotificationId LobbyInviteAcceptedId;
 	FCallbackBase* LobbyInviteAcceptedCallback;
 	EOS_NotificationId JoinLobbyAcceptedId;
@@ -250,8 +252,9 @@ private:
 	void OnLobbyUpdateReceived(const EOS_LobbyId& LobbyId);
 	void OnLobbyMemberUpdateReceived(const EOS_LobbyId& LobbyId, const EOS_ProductUserId& TargetUserId);
 	void OnMemberStatusReceived(const EOS_LobbyId& LobbyId, const EOS_ProductUserId& TargetUserId, EOS_ELobbyMemberStatus CurrentStatus);
-	void OnLobbyInviteAccepted(const char* InviteId, const EOS_ProductUserId& LocalUserId, const EOS_ProductUserId& TargetUserId);
-	void OnJoinLobbyAccepted(const EOS_ProductUserId& LocalUserId, const EOS_UI_EventId& UiEventId);
+	void OnLobbyInviteReceived(const EOS_Lobby_LobbyInviteReceivedCallbackInfo* Data);
+	void OnLobbyInviteAccepted(const EOS_Lobby_LobbyInviteAcceptedCallbackInfo* Data);
+	void OnJoinLobbyAccepted(const EOS_Lobby_JoinLobbyAcceptedCallbackInfo* Data);
 
 	// Methods to update an API Lobby from an OSS Lobby
 	void SetLobbyPermissionLevel(EOS_HLobbyModification LobbyModificationHandle, FNamedOnlineSession* Session);
@@ -274,6 +277,11 @@ private:
 	// Helper methods
 	typedef TFunction<void(const EOS_ProductUserId& ProductUserId, EOS_EpicAccountId& EpicAccountId)> GetEpicAccountIdAsyncCallback;
 
+	void RegisterSessionNotifications();
+	void OnSessionInviteReceived(const EOS_Sessions_SessionInviteReceivedCallbackInfo* Data);
+	void OnSessionInviteAccepted(const EOS_Sessions_SessionInviteAcceptedCallbackInfo* Data);
+	void OnJoinSessionAccepted(const EOS_Sessions_JoinSessionAcceptedCallbackInfo* Data);
+
 	void RegisterLobbyNotifications();
 	FNamedOnlineSession* GetNamedSessionFromLobbyId(const FUniqueNetIdEOSLobby& LobbyId);
 	FOnlineSessionSearchResult* GetSearchResultFromLobbyId(const FUniqueNetIdEOSLobby& LobbyId);
@@ -341,8 +349,12 @@ private:
 	TSharedPtr<FOnlineSessionSearch> LastInviteSearch;
 
 	/** Notification state for SDK events */
+	EOS_NotificationId SessionInviteReceivedId;
+	FCallbackBase* SessionInviteReceivedCallback;
 	EOS_NotificationId SessionInviteAcceptedId;
 	FCallbackBase* SessionInviteAcceptedCallback;
+	EOS_NotificationId JoinSessionAcceptedId;
+	FCallbackBase* JoinSessionAcceptedCallback;
 
 	bool bIsUsingP2PSockets;
 };
