Refactor of UserManagerEOS methods to handle UniqueNetId resolution before performing any friends or user operations.
Fix for an error that didn't save Friends user attributes accordingly. Added logic to pass them to the friend type constructor.
Removal of superfluous logic after refactor.

diff --git a/Engine/Plugins/Online/OnlineSubsystemEOS/Source/OnlineSubsystemEOS/Private/OnlineSubsystemEOSTypes.h b/Engine/Plugins/Online/OnlineSubsystemEOS/Source/OnlineSubsystemEOS/Private/OnlineSubsystemEOSTypes.h
--- a/Engine/Plugins/Online/OnlineSubsystemEOS/Source/OnlineSubsystemEOS/Private/OnlineSubsystemEOSTypes.h
+++ b/Engine/Plugins/Online/OnlineSubsystemEOS/Source/OnlineSubsystemEOS/Private/OnlineSubsystemEOSTypes.h
@@ -165,6 +165,12 @@ public:
 	{
 	}
 
+	TOnlineUserEOS(const FUniqueNetIdEOSRef& InNetIdRef, const TMap<FString, FString>& InUserAttributes)
+		: UserIdRef(InNetIdRef)
+		, UserAttributes(InUserAttributes)
+	{
+	}
+
 	virtual ~TOnlineUserEOS()
 	{
 	}
@@ -294,6 +300,11 @@ public:
 	{
 	}
 
+	TOnlineFriendEOS(const FUniqueNetIdEOSRef& InNetIdRef, const TMap<FString, FString>& InUserAttributes)
+		: TOnlineUserEOS<BaseClass, IAttributeAccessInterface>(InNetIdRef, InUserAttributes)
+	{
+	}
+
 // FOnlineFriend
 	/**
 	 * @return the current invite status of a friend wrt to user that queried
diff --git a/Engine/Plugins/Online/OnlineSubsystemEOS/Source/OnlineSubsystemEOS/Private/UserManagerEOS.cpp b/Engine/Plugins/Online/OnlineSubsystemEOS/Source/OnlineSubsystemEOS/Private/UserManagerEOS.cpp
--- a/Engine/Plugins/Online/OnlineSubsystemEOS/Source/OnlineSubsystemEOS/Private/UserManagerEOS.cpp
+++ b/Engine/Plugins/Online/OnlineSubsystemEOS/Source/OnlineSubsystemEOS/Private/UserManagerEOS.cpp
@@ -1053,7 +1053,8 @@ FLocalUserEOS& FUserManagerEOS::AddLocalUser(int32 LocalUserNum, EOS_EpicAccount
 		NotificationPairPtr->NotificationId = EOS_Connect_AddNotifyAuthExpiration(EOSSubsystem->ConnectHandle, &Options, CallbackObj, CallbackObj->GetCallbackPtr());
 	}
 
-	UniqueNetIdToAttributeAccessMap.Emplace(UserNetId, UserAccountRef);
+	const FOnlineUserEOSRef UserRef(new FOnlineUserEOS(UserNetId));
+	UniqueNetIdToUserRefMap.Emplace(UserNetId, UserRef);
 
 	// Once all the fields are set, we start gathering additional information
 	ReadFriendsList(LocalUserNum, FString());
@@ -1096,24 +1097,6 @@ void FUserManagerEOS::UpdateUserInfo(IAttributeAccessInterfaceRef AttributeAcces
 	}
 }
 
-bool FUserManagerEOS::IsFriendQueryUserInfoOngoing(int32 LocalUserNum)
-{
-	const FLocalUserEOS& LocalUser = GetLocalUserChecked(LocalUserNum);
-
-	// If we have an entry for this user and the corresponding array has any element, users are still being processed
-	if (LocalUser.OngoingFriendQueryUserInfo.Num() > 0)
-	{
-		return true;
-	}
-	
-	if (LocalUser.OngoingPlayerQueryExternalMappings.Num() > 0)
-	{
-		return true;
-	}
-
-	return false;
-}
-
 TSharedPtr<FUserOnlineAccount> FUserManagerEOS::GetUserAccount(const FUniqueNetId& UserId) const
 {
 	TSharedPtr<FUserOnlineAccount> Result;
@@ -1413,6 +1396,67 @@ void FUserManagerEOS::ResolveUniqueNetIds(int32 LocalUserNum, const TArray<EOS_P
 	}
 }
 
+void FUserManagerEOS::ResolveUniqueNetId(int32 LocalUserNum, const EOS_EpicAccountId& EpicAccountId, const FResolveUniqueNetIdCallback& Callback)
+{
+	const FResolveEpicAccountIdsCallback GroupCallback = [EpicAccountId, OriginalCallback = Callback](bool bWasSuccessful, TMap<EOS_EpicAccountId, FUniqueNetIdEOSRef> ResolvedUniqueNetIds, FString ErrorStr) {
+		OriginalCallback(ResolvedUniqueNetIds[EpicAccountId]);
+	};
+
+	ResolveUniqueNetIds(LocalUserNum, { EpicAccountId }, GroupCallback);
+}
+
+void FUserManagerEOS::ResolveUniqueNetIds(int32 LocalUserNum, const TArray<EOS_EpicAccountId>& EpicAccountIds, const FResolveEpicAccountIdsCallback& Callback)
+{
+	TMap<EOS_EpicAccountId, FUniqueNetIdEOSRef> ResolvedUniqueNetIds;
+	TArray<FString> EpicAccountIdsToResolve;
+
+	for (const EOS_EpicAccountId& EpicAccountId : EpicAccountIds)
+	{
+		if (ensure(EOS_EpicAccountId_IsValid(EpicAccountId) == EOS_TRUE))
+		{
+			// We check first if the id is already registered, which would allow us to retrieve it directly
+			if (FUniqueNetIdEOSPtr UniqueNetId = FUniqueNetIdEOSRegistry::Find(EpicAccountId))
+			{
+				ResolvedUniqueNetIds.Add(EpicAccountId, UniqueNetId.ToSharedRef());
+			}
+			else
+			{
+				// If that's not the case, we'll have to query them first
+				EpicAccountIdsToResolve.Add(LexToString(EpicAccountId));
+			}
+		}
+	}
+
+	if (!EpicAccountIdsToResolve.IsEmpty())
+	{
+		QueryExternalIdMappings(*GetLocalUniqueNetIdEOS(LocalUserNum), FExternalIdQueryOptions(), EpicAccountIdsToResolve, FOnQueryExternalIdMappingsComplete::CreateLambda([this, ResolvedUniqueNetIds = MoveTemp(ResolvedUniqueNetIds), Callback](bool bWasSuccessful, const FUniqueNetId& UserId, const FExternalIdQueryOptions& QueryOptions, const TArray<FString>& ExternalIds, const FString& Error) mutable
+			{
+				if (bWasSuccessful)
+				{
+					for (const FString& ExternalId : ExternalIds)
+					{
+						EOS_EpicAccountId AccountId;
+						LexFromString(AccountId, *ExternalId);
+
+						if (EOS_EpicAccountId_IsValid(AccountId) == EOS_TRUE)
+						{
+							ResolvedUniqueNetIds.Add(AccountId, FUniqueNetIdEOSRegistry::FindChecked(AccountId));
+						}
+					}
+				}
+
+				Callback(bWasSuccessful, ResolvedUniqueNetIds, Error);
+			}));
+	}
+	else
+	{
+		EOSSubsystem->ExecuteNextTick([Callback, ResolvedUniqueNetIds = MoveTemp(ResolvedUniqueNetIds)]()
+			{
+				Callback(true, ResolvedUniqueNetIds, FString());
+			});
+	}
+}
+
 FOnlineUserPtr FUserManagerEOS::GetLocalOnlineUser(int32 LocalUserNum) const
 {
 	FOnlineUserPtr OnlineUser;
@@ -1770,147 +1814,185 @@ typedef TEOSCallback<EOS_Friends_OnQueryFriendsCallback, EOS_Friends_QueryFriend
 
 void FUserManagerEOS::FriendStatusChanged(const EOS_Friends_OnFriendsUpdateInfo* Data)
 {
-	// Get the local user information
 	if (IsLocalUser(Data->LocalUserId))
 	{
 		const int32 LocalUserNum = GetLocalUserNumFromEpicAccountId(Data->LocalUserId);
-		const FUniqueNetIdEOSRef LocalEOSID = GetLocalUserChecked(LocalUserNum).UniqueNetId.ToSharedRef();
-		if (!FUniqueNetIdEOSRegistry::Find(Data->TargetUserId))
+
+		// If the remote user for which we just got an update is not registered, we will register it first
+		FUniqueNetIdEOSPtr NetId = FUniqueNetIdEOSRegistry::Find(Data->TargetUserId);
+		bool bIsNetIdRegistered = NetId.IsValid() ? UniqueNetIdToUserRefMap.Contains(NetId.ToSharedRef()) : false;
+		if (!bIsNetIdRegistered)
+		{
+			FRemoteUserProcessedCallback Callback = [this, LocalUserId = Data->LocalUserId, TargetUserId = Data->TargetUserId, PreviousStatus = Data->PreviousStatus, CurrentStatus = Data->CurrentStatus](bool bWasSuccessful, FUniqueNetIdEOSRef RemotePlayerNetId, const FString& ErrorStr)
+			{
+				FriendStatusChangedImpl(LocalUserId, TargetUserId, PreviousStatus, CurrentStatus);
+			};
+
+			AddRemotePlayer(LocalUserNum, Data->TargetUserId, Callback);
+		}
+		else
 		{
-			AddRemotePlayer(LocalUserNum, Data->TargetUserId);
-			// TODO: We need to modify AddRemotePlayer to take a delegate, to execute after all information has been retrieved
+			FriendStatusChangedImpl(Data->LocalUserId, Data->TargetUserId, Data->PreviousStatus, Data->CurrentStatus);
 		}
+	}
+}
 
-		// AddRemotePlayer just added this id to the registry, so it will be found and valid, no checks needed
-		const FUniqueNetIdEOSRef& FriendEOSId = FUniqueNetIdEOSRegistry::FindChecked(Data->TargetUserId);
+void FUserManagerEOS::FriendStatusChangedImpl(EOS_EpicAccountId LocalUserId, EOS_EpicAccountId TargetUserId, EOS_EFriendsStatus PreviousStatus, EOS_EFriendsStatus CurrentStatus)
+{
+	const int32 LocalUserNum = GetLocalUserNumFromEpicAccountId(LocalUserId);
+	const FUniqueNetIdEOSRef LocalEOSID = GetLocalUserChecked(LocalUserNum).UniqueNetId.ToSharedRef();
+	const FUniqueNetIdEOSRef& FriendEOSId = FUniqueNetIdEOSRegistry::FindChecked(TargetUserId);
 
-		switch (Data->CurrentStatus)
-		{
-		case EOS_EFriendsStatus::EOS_FS_NotFriends: // Invite rejections and friend removal
+	switch (CurrentStatus)
+	{
+	case EOS_EFriendsStatus::EOS_FS_NotFriends: // Invite rejections and friend removal
+	{
+		//User should already be a friend
+		FFriendsListEOSPtr& FriendsListRef = GetLocalUserChecked(LocalUserNum).FriendsList;
+		FOnlineFriendEOSPtr Friend = FriendsListRef->GetByNetId(FriendEOSId);
+		if (Friend.IsValid())
 		{
-			//User should already be a friend
-			FFriendsListEOSPtr& FriendsListRef = GetLocalUserChecked(LocalUserNum).FriendsList;
-			FOnlineFriendEOSPtr Friend = FriendsListRef->GetByNetId(FriendEOSId);
-			if (Friend.IsValid())
-			{
-				FriendsListRef->Remove(FriendEOSId, Friend.ToSharedRef());
-				Friend->SetInviteStatus(EInviteStatus::Unknown);
+			FriendsListRef->Remove(FriendEOSId, Friend.ToSharedRef());
+			Friend->SetInviteStatus(EInviteStatus::Unknown);
 
-				if (Data->PreviousStatus == EOS_EFriendsStatus::EOS_FS_Friends)
-				{
-					TriggerOnFriendRemovedDelegates(*LocalEOSID, *FriendEOSId);
-				}
-				else if (Data->PreviousStatus == EOS_EFriendsStatus::EOS_FS_InviteSent || Data->PreviousStatus == EOS_EFriendsStatus::EOS_FS_InviteReceived)
-				{
-					TriggerOnInviteRejectedDelegates(*LocalEOSID, *FriendEOSId); // We don't have an "OnInviteRejected" event only for the local user
-				}
+			if (PreviousStatus == EOS_EFriendsStatus::EOS_FS_Friends)
+			{
+				TriggerOnFriendRemovedDelegates(*LocalEOSID, *FriendEOSId);
 			}
-			else
+			else if (PreviousStatus == EOS_EFriendsStatus::EOS_FS_InviteSent || PreviousStatus == EOS_EFriendsStatus::EOS_FS_InviteReceived)
 			{
-				UE_LOG_ONLINE_FRIEND(Verbose, TEXT("Friend status notification received for user [%d], but remote user [%s] was not previously registered as a friend"), *FriendEOSId->ToString());
+				TriggerOnInviteRejectedDelegates(*LocalEOSID, *FriendEOSId); // We don't have an "OnInviteRejected" event only for the local user
 			}
-
-			break;
 		}
-		case EOS_EFriendsStatus::EOS_FS_InviteSent: // Invite sent by local user
+		else
 		{
-			// The only supported case is NotFriends->InviteSent, other combinations are not possible
-			check(Data->PreviousStatus == EOS_EFriendsStatus::EOS_FS_NotFriends);
+			UE_LOG_ONLINE_FRIEND(Verbose, TEXT("Friend status notification received for user [%d], but remote user [%s] was not previously registered as a friend"), *FriendEOSId->ToString());
+		}
 
-			FOnlineFriendEOSRef FriendRef = AddFriend(LocalUserNum, Data->TargetUserId);
+		break;
+	}
+	case EOS_EFriendsStatus::EOS_FS_InviteSent: // Invite sent by local user
+	{
+		// The only supported case is NotFriends->InviteSent, other combinations are not possible
+		check(PreviousStatus == EOS_EFriendsStatus::EOS_FS_NotFriends);
 
-			FriendRef->SetInviteStatus(EInviteStatus::PendingOutbound);
-			TriggerOnOutgoingInviteSentDelegates(LocalUserNum);
+		FOnlineFriendEOSRef FriendRef = AddFriend(LocalUserNum, *FriendEOSId);
 
-			break;
-		}
-		case EOS_EFriendsStatus::EOS_FS_InviteReceived: // Invite received by local user
-		{
-			// The only supported case is NotFriends->InviteReceived, other combinations are not possible
-			check(Data->PreviousStatus == EOS_EFriendsStatus::EOS_FS_NotFriends);
+		FriendRef->SetInviteStatus(EInviteStatus::PendingOutbound);
+		TriggerOnOutgoingInviteSentDelegates(LocalUserNum);
 
-			FOnlineFriendEOSRef FriendRef = AddFriend(LocalUserNum, Data->TargetUserId);
+		break;
+	}
+	case EOS_EFriendsStatus::EOS_FS_InviteReceived: // Invite received by local user
+	{
+		// The only supported case is NotFriends->InviteReceived, other combinations are not possible
+		check(PreviousStatus == EOS_EFriendsStatus::EOS_FS_NotFriends);
 
-			FriendRef->SetInviteStatus(EInviteStatus::PendingInbound);
-			TriggerOnInviteReceivedDelegates(*LocalEOSID, *FriendEOSId);
+		FOnlineFriendEOSRef FriendRef = AddFriend(LocalUserNum, *FriendEOSId);
 
-			break;
-		}
-		case EOS_EFriendsStatus::EOS_FS_Friends: // Accepted invites and friend addition. This logic applies to all three PreviousStatus cases
-		{
-			FOnlineFriendEOSRef FriendRef = AddFriend(LocalUserNum, Data->TargetUserId);
+		FriendRef->SetInviteStatus(EInviteStatus::PendingInbound);
+		TriggerOnInviteReceivedDelegates(*LocalEOSID, *FriendEOSId);
 
-			FriendRef->SetInviteStatus(EInviteStatus::Accepted);
-			TriggerOnInviteAcceptedDelegates(*LocalEOSID, *FriendEOSId); // We don't have an "OnFriendAdded", all relationships exist in the context of invites. We also don't have a separate "OnInviteAccepted" for the local user.
+		break;
+	}
+	case EOS_EFriendsStatus::EOS_FS_Friends: // Accepted invites and friend addition. This logic applies to all three PreviousStatus cases
+	{
+		FOnlineFriendEOSRef FriendRef = AddFriend(LocalUserNum, *FriendEOSId);
 
-			break;
-		}
-		default:
-			checkNoEntry();
-			UE_LOG_ONLINE_FRIEND(Warning, TEXT("[FUserManagerEOS::FriendStatusChanged] Unsupported status received as CurrentStatus"));
-		}
+		FriendRef->SetInviteStatus(EInviteStatus::Accepted);
+		TriggerOnInviteAcceptedDelegates(*LocalEOSID, *FriendEOSId); // We don't have an "OnFriendAdded", all relationships exist in the context of invites. We also don't have a separate "OnInviteAccepted" for the local user.
 
-		TriggerOnFriendsChangeDelegates(LocalUserNum);
+		break;
 	}
+	default:
+		checkNoEntry();
+		UE_LOG_ONLINE_FRIEND(Warning, TEXT("[FUserManagerEOS::FriendStatusChanged] Unsupported status received as CurrentStatus"));
+	}
+
+	TriggerOnFriendsChangeDelegates(LocalUserNum);
 }
 
-FOnlineFriendEOSRef FUserManagerEOS::AddFriend(int32 LocalUserNum, EOS_EpicAccountId EpicAccountId)
+FOnlineFriendEOSRef FUserManagerEOS::AddFriend(int32 LocalUserNum, const FUniqueNetIdEOS& FriendNetId)
 {
-	FUniqueNetIdEOSRef FriendNetId = FUniqueNetIdEOSRegistry::FindOrAdd(EpicAccountId, nullptr);
-	FOnlineFriendEOSRef FriendRef = MakeShareable(new FOnlineFriendEOS(FriendNetId));
+	const FOnlineUserEOSRef UserRef = UniqueNetIdToUserRefMap[FriendNetId.AsShared()];
+	const FUniqueNetIdEOSRef FriendNetIdEOSRef = StaticCastSharedRef<const FUniqueNetIdEOS>(FriendNetId.AsShared());
+	const FOnlineFriendEOSRef FriendRef = MakeShareable(new FOnlineFriendEOS(FriendNetIdEOSRef, UserRef->UserAttributes));
 
-	GetLocalUserChecked(LocalUserNum).FriendsList->Add(FriendNetId, FriendRef);
+	GetLocalUserChecked(LocalUserNum).FriendsList->Add(FriendNetId.AsShared(), FriendRef);
 
 	EOS_Friends_GetStatusOptions Options = { };
 	Options.ApiVersion = 1;
 	UE_EOS_CHECK_API_MISMATCH(EOS_FRIENDS_GETSTATUS_API_LATEST, 1);
 	Options.LocalUserId = GetLocalEpicAccountId(LocalUserNum);
-	Options.TargetUserId = EpicAccountId;
+	Options.TargetUserId = FriendNetId.GetEpicAccountId();
 	EOS_EFriendsStatus Status = EOS_Friends_GetStatus(EOSSubsystem->FriendsHandle, &Options);
 	
 	FriendRef->SetInviteStatus(ToEInviteStatus(Status));
 
-	// Add this friend as a remote player if we haven't already (this will grab user info)
-	if (!UniqueNetIdToAttributeAccessMap.Contains(FriendNetId))
-	{
-		AddRemotePlayer(LocalUserNum, EpicAccountId, FriendRef);
-	}
-
 	// Querying the presence of a non-friend would cause an SDK error.
 	// Players that sent/received a friend invitation from us still count as "friends", so check
 	// our friend relationship here.
 	if(Status == EOS_EFriendsStatus::EOS_FS_Friends)
 	{
-		QueryPresence(*FriendNetId, IgnoredPresenceDelegate);
+		QueryPresence(FriendNetId, IgnoredPresenceDelegate);
 	}
 
 	return FriendRef;
 }
 
-void FUserManagerEOS::AddRemotePlayer(int32 LocalUserNum, EOS_EpicAccountId EpicAccountId)
+void FUserManagerEOS::AddRemotePlayer(int32 LocalUserNum, EOS_EpicAccountId EpicAccountId, const FRemoteUserProcessedCallback& Callback)
 {
-	FUniqueNetIdEOSRef EOSID = FUniqueNetIdEOSRegistry::FindOrAdd(EpicAccountId, nullptr);
-	FOnlineUserEOSRef UserRef = MakeShareable(new FOnlineUserEOS(EOSID));
+	const FResolveUniqueNetIdCallback IdResolutionCallback = [this, LocalUserNum, EpicAccountId, Callback](FUniqueNetIdEOSRef ResolvedUniqueNetId) mutable
+	{
+		const FOnlineUserEOSRef AttributeRef = MakeShareable(new FOnlineUserEOS(ResolvedUniqueNetId));
 
-	// Add this user as a remote (this will grab presence & user info)
-	AddRemotePlayer(LocalUserNum, EpicAccountId, UserRef);
-}
+		UniqueNetIdToUserRefMap.Emplace(ResolvedUniqueNetId, AttributeRef);
 
-void FUserManagerEOS::AddRemotePlayer(int32 LocalUserNum, EOS_EpicAccountId EpicAccountId, IAttributeAccessInterfaceRef AttributeRef)
-{
-	// TODO: pending asynchronous resolution of the remote player's UniqueNetId
-	FUniqueNetIdEOSRef EOSID = FUniqueNetIdEOSRegistry::FindOrAdd(EpicAccountId, nullptr);
-	UniqueNetIdToAttributeAccessMap.Emplace(EOSID, AttributeRef);
+		// Read the user info for this player
+		ReadUserInfo(LocalUserNum, EpicAccountId, Callback);
+	};
 
-	// Read the user info for this player
-	ReadUserInfo(LocalUserNum, EpicAccountId);
+	ResolveUniqueNetId(LocalUserNum, EpicAccountId, IdResolutionCallback);
 }
 
-void FUserManagerEOS::UpdateRemotePlayerProductUserId(EOS_EpicAccountId EpicAccountId, EOS_ProductUserId ProductUserId)
+void FUserManagerEOS::AddRemotePlayers(int32 LocalUserNum, TArray<EOS_EpicAccountId> EpicAccountIds, const FRemoteUsersProcessedCallback& Callback)
 {
-	// TODO: Unique Id updates are removed, but we still need to add proper id resolution for all ids to keep a correct state
+	const FResolveEpicAccountIdsCallback IdResolutionCallback = [this, LocalUserNum, EpicAccountIds, Callback](bool bWasSuccessful, TMap<EOS_EpicAccountId, FUniqueNetIdEOSRef> ResolvedUniqueNetIds, FString ErrorStr) mutable
+	{
+		for (const TPair<EOS_EpicAccountId, FUniqueNetIdEOSRef>& Entry : ResolvedUniqueNetIds)
+		{
+			const FOnlineUserEOSRef AttributeRef = MakeShareable(new FOnlineUserEOS(Entry.Value));
+
+			UniqueNetIdToUserRefMap.Emplace(Entry.Value, AttributeRef);
+
+			const FRemoteUserProcessedCallback ReadUserInfoCallback = [this, LocalUserNum, Callback](bool bWasSuccessful, FUniqueNetIdEOSRef UserNetId, const FString& ErrorStr) mutable
+			{
+				FLocalUserEOS& LocalUser = GetLocalUserChecked(LocalUserNum);
+
+				LocalUser.OngoingQueryUserInfoResults.ProcessedIds.Add(UserNetId);
+
+				LocalUser.OngoingQueryUserInfoResults.bAllWasSuccessful &= bWasSuccessful;
 
-	FUniqueNetIdEOSRegistry::FindOrAdd(EpicAccountId, ProductUserId);	
+				if (!ErrorStr.IsEmpty())
+				{
+					LocalUser.OngoingQueryUserInfoResults.AllErrorStr += TEXT("/n") + ErrorStr;
+				}
+
+				if (LocalUser.OngoingQueryUserInfoAccounts.IsEmpty())
+				{
+					FLocalUserEOS::FReadUserInfoResults Results = LocalUser.OngoingQueryUserInfoResults;
+					LocalUser.OngoingQueryUserInfoResults.Reset();
+					
+					Callback(Results.bAllWasSuccessful, Results.ProcessedIds, Results.AllErrorStr);
+				}
+			};
+
+			// Read the user info for this player
+			ReadUserInfo(LocalUserNum, Entry.Key, ReadUserInfoCallback);
+		}		
+	};
+
+	ResolveUniqueNetIds(LocalUserNum, EpicAccountIds, IdResolutionCallback);
 }
 
 // IOnlineFriends Interface
@@ -1960,9 +2042,7 @@ bool FUserManagerEOS::ReadFriendsList(int32 LocalUserNum, const FString& ListNam
 			Options.LocalUserId = GetLocalEpicAccountId(LocalUserNum);
 			int32 FriendCount = EOS_Friends_GetFriendsCount(EOSSubsystem->FriendsHandle, &Options);
 
-			GetLocalUserChecked(LocalUserNum).FriendsList->Empty(FriendCount);
-
-			TArray<FString> FriendEasIds;
+			TArray<EOS_EpicAccountId> FriendEasIds;
 			FriendEasIds.Reserve(FriendCount);
 			// Process each friend returned
 			for (int32 Index = 0; Index < FriendCount; Index++)
@@ -1973,26 +2053,29 @@ bool FUserManagerEOS::ReadFriendsList(int32 LocalUserNum, const FString& ListNam
 				FriendIndexOptions.Index = Index;
 				FriendIndexOptions.LocalUserId = Options.LocalUserId;
 				EOS_EpicAccountId FriendEpicAccountId = EOS_Friends_GetFriendAtIndex(EOSSubsystem->FriendsHandle, &FriendIndexOptions);
-				if (FriendEpicAccountId != nullptr)
+				if (EOS_EpicAccountId_IsValid(FriendEpicAccountId) == EOS_TRUE)
 				{
-					AddFriend(LocalUserNum, FriendEpicAccountId);
-					FriendEasIds.Add(LexToString(FriendEpicAccountId));
+					FriendEasIds.Add(FriendEpicAccountId);
 				}
 			}
 
-			const TFunction<FOnQueryExternalIdMappingsComplete::TFuncType>& OnExternalIdMappingsQueriedLambda =
-				[this, WeakThis = AsWeak(), LocalUserNum](bool bWasSuccessful, const FUniqueNetId& UserId, const FExternalIdQueryOptions& QueryOptions, const TArray<FString>& ExternalIds, const FString& Error)
+			const FRemoteUsersProcessedCallback Callback = [this, LocalUserNum, FriendCount](bool bWasSuccessful, TArray<FUniqueNetIdEOSRef> RemoteUserNetIds, const FString& ErrorStr)
 			{
-				if (FUserManagerEOSPtr StrongThis = WeakThis.Pin())
+				if (bWasSuccessful)
 				{
-					ProcessReadFriendsListComplete(LocalUserNum, bWasSuccessful, Error);
-				}
-			};
+					GetLocalUserChecked(LocalUserNum).FriendsList->Empty(FriendCount);
 
-			const auto& ExternalMappingsCallback = OSSInternalCallback::Create<FOnQueryExternalIdMappingsComplete>(EOSSubsystem->UserManager, OnExternalIdMappingsQueriedLambda);
+					for (const FUniqueNetIdEOSRef& RemoteUserNetId : RemoteUserNetIds)
+					{
+						AddFriend(LocalUserNum, *RemoteUserNetId);
+					}
+				}
 
-			QueryExternalIdMappings(*GetLocalUniqueNetIdEOS(DefaultLocalUser), FExternalIdQueryOptions(), FriendEasIds, ExternalMappingsCallback);
+				// We trigger all the corresponding delegates
+				ProcessReadFriendsListComplete(LocalUserNum, bWasSuccessful, ErrorStr);
+			};
 
+			AddRemotePlayers(LocalUserNum, FriendEasIds, Callback);
 		}
 		else
 		{
@@ -2000,11 +2083,30 @@ bool FUserManagerEOS::ReadFriendsList(int32 LocalUserNum, const FString& ListNam
 			ProcessReadFriendsListComplete(LocalUserNum, false, ErrorString);
 		}
 	};
+
 	EOS_Friends_QueryFriends(EOSSubsystem->FriendsHandle, &Options, CallbackObj, CallbackObj->GetCallbackPtr());
 
 	return true;
 }
 
+bool FUserManagerEOS::IsFriendQueryUserInfoOngoing(int32 LocalUserNum)
+{
+	const FLocalUserEOS& LocalUser = GetLocalUserChecked(LocalUserNum);
+
+	// If we have an entry for this user and the corresponding array has any element, users are still being processed
+	if (LocalUser.OngoingQueryUserInfoAccounts.Num() > 0)
+	{
+		return true;
+	}
+
+	if (LocalUser.OngoingPlayerQueryExternalMappings.Num() > 0)
+	{
+		return true;
+	}
+
+	return false;
+}
+
 void FUserManagerEOS::ProcessReadFriendsListComplete(int32 LocalUserNum, bool bWasSuccessful, const FString& ErrorStr)
 {
 	// If we started any user info queries for friends, we'll just wait until they finish
@@ -2019,10 +2121,6 @@ void FUserManagerEOS::ProcessReadFriendsListComplete(int32 LocalUserNum, bool bW
 		}
 		LocalUser.CachedReadUserListInfo.Empty();
 
-		// These should be empty by now, but we make sure
-		LocalUser.OngoingFriendQueryUserInfo.Empty();
-		LocalUser.OngoingPlayerQueryExternalMappings.Empty();
-
 		TriggerOnFriendsChangeDelegates(LocalUserNum);
 	}
 }
@@ -2590,20 +2688,20 @@ void FUserManagerEOS::QueryPresence(const FUniqueNetId& UserId, const FOnPresenc
 		FQueryPresenceCallback* CallbackObj = new FQueryPresenceCallback(AsWeak());
 		CallbackObj->CallbackLambda = [this, Delegate](const EOS_Presence_QueryPresenceCallbackInfo* Data)
 		{
+			// If we were able to retrieve a UniqueNetId to pass to this method, it must already be registered
+			const FUniqueNetIdEOSRef NetIdPtr = FUniqueNetIdEOSRegistry::FindChecked(Data->TargetUserId);
+
 			if (Data->ResultCode == EOS_EResult::EOS_Success)
 			{
-				if (const FUniqueNetIdEOSPtr NetIdPtr = FUniqueNetIdEOSRegistry::Find(Data->TargetUserId))
-				{
-					// Update the presence data to the most recent
-					UpdatePresence(DefaultLocalUser, Data->TargetUserId);
-					Delegate.ExecuteIfBound(*NetIdPtr, true);
-					return;
-				}
+				// Update the presence data to the most recent
+				UpdatePresence(DefaultLocalUser, Data->TargetUserId);
+				Delegate.ExecuteIfBound(*NetIdPtr, true);
+			}
+			else
+			{
+				UE_LOG_ONLINE(Error, TEXT("QueryPresence() for user (%s) failed with result code (%s)"), *NetIdPtr->ToString(), *LexToString(Data->ResultCode));
+				Delegate.ExecuteIfBound(*FUniqueNetIdEOS::EmptyId(), false);
 			}
-
-			const FString& TargetUser = FUniqueNetIdEOSRegistry::FindOrAdd(Data->TargetUserId, nullptr)->ToString();
-			UE_LOG_ONLINE(Error, TEXT("QueryPresence() for user (%s) failed with result code (%s)"), *TargetUser, *LexToString(Data->ResultCode));
-			Delegate.ExecuteIfBound(*FUniqueNetIdEOS::EmptyId(), false);
 		};
 
 		// Query for updated presence
@@ -2712,6 +2810,13 @@ EOnlineCachedResult::Type FUserManagerEOS::GetCachedPresenceForApp(const FUnique
 
 bool FUserManagerEOS::QueryUserInfo(int32 LocalUserNum, const TArray<FUniqueNetIdRef>& UserIds)
 {
+	const FLocalUserEOS& LocalUser = GetLocalUserChecked(LocalUserNum);
+	if (!LocalUser.OngoingQueryUserInfoAccounts.IsEmpty())
+	{
+		UE_LOG_ONLINE(Verbose, TEXT("[FUserManagerEOS::QueryUserInfo] User Info query already ongoing for local user %d"), LocalUserNum);
+		TriggerOnQueryUserInfoCompleteDelegates(LocalUserNum, false, TArray<FUniqueNetIdRef>(), FString());
+	}
+
 	TArray<FString> UserEasIdsNeedingExternalMappings;
 	UserEasIdsNeedingExternalMappings.Reserve(UserIds.Num());
 	
@@ -2719,19 +2824,34 @@ bool FUserManagerEOS::QueryUserInfo(int32 LocalUserNum, const TArray<FUniqueNetI
 	for (const FUniqueNetIdRef& NetId : UserIds)
 	{
 		const FUniqueNetIdEOS& EOSID = FUniqueNetIdEOS::Cast(*NetId);
+
 		// Skip querying for local users since we already have that data
 		if (IsLocalUser(EOSID))
 		{
 			continue;
 		}
+
 		// Check to see if we know about this user or not
 		const EOS_EpicAccountId AccountId = EOSID.GetEpicAccountId();
 		if (EOS_EpicAccountId_IsValid(AccountId) == EOS_TRUE)
 		{
+			const FRemoteUserProcessedCallback Callback = [this, LocalUserNum](bool bWasSuccessful, FUniqueNetIdEOSRef RemoteUserNetId, const FString& ErrorStr)
+			{
+				static TArray<FUniqueNetIdRef> ProcessedIds;
+				ProcessedIds.Add(RemoteUserNetId);
+
+				if(GetLocalUserChecked(LocalUserNum).OngoingQueryUserInfoAccounts.IsEmpty())
+				{
+					TriggerOnQueryUserInfoCompleteDelegates(LocalUserNum, bWasSuccessful, ProcessedIds, ErrorStr);
+					
+					ProcessedIds.Empty();
+				}
+			};
+
 			// If the user is already registered, we'll update their user info
-			if (UniqueNetIdToAttributeAccessMap.Contains(EOSID.AsShared()))
+			if (UniqueNetIdToUserRefMap.Contains(EOSID.AsShared()))
 			{
-				ReadUserInfo(LocalUserNum, AccountId);
+				ReadUserInfo(LocalUserNum, AccountId, Callback);
 			}
 			else
 			{
@@ -2739,46 +2859,29 @@ bool FUserManagerEOS::QueryUserInfo(int32 LocalUserNum, const TArray<FUniqueNetI
 				UserEasIdsNeedingExternalMappings.Add(LexToString(AccountId));
 
 				// Registering the player will also query the user info data
-				AddRemotePlayer(LocalUserNum, AccountId);
+				AddRemotePlayer(LocalUserNum, AccountId, Callback);
 			}
 		}
 	}
-
-	const FUniqueNetIdEOSPtr LocalId = GetLocalUniqueNetIdEOS(LocalUserNum);
-	QueryExternalIdMappings(*LocalId, FExternalIdQueryOptions(), UserEasIdsNeedingExternalMappings, FOnQueryExternalIdMappingsComplete::CreateLambda([this, WeakThis = AsWeak(), LocalUserNum](bool bWasSuccessful, const FUniqueNetId& UserId, const FExternalIdQueryOptions& QueryOptions, const TArray<FString>& ExternalIds, const FString& ErrorStr)
-		{
-			if (FUserManagerEOSPtr StrongThis = WeakThis.Pin())
-			{
-				TArray<FUniqueNetIdPtr> ExternalIdPtrs;
-				GetExternalIdMappings(QueryOptions, ExternalIds, ExternalIdPtrs);
-
-				TArray<FUniqueNetIdRef> ExternalIdRefs;
-				for (const FUniqueNetIdPtr& ExternalIdPtr : ExternalIdPtrs)
-				{
-					if (ExternalIdPtr.IsValid())
-					{
-						ExternalIdRefs.Add(ExternalIdPtr.ToSharedRef());
-					}
-				}
-
-				TriggerOnQueryUserInfoCompleteDelegates(LocalUserNum, bWasSuccessful, ExternalIdRefs, ErrorStr);
-			}
-		}));
 	
 	return true;
 }
 
 typedef TEOSCallback<EOS_UserInfo_OnQueryUserInfoCallback, EOS_UserInfo_QueryUserInfoCallbackInfo, FUserManagerEOS> FReadUserInfoCallback;
 
-void FUserManagerEOS::ReadUserInfo(int32 LocalUserNum, EOS_EpicAccountId EpicAccountId)
+void FUserManagerEOS::ReadUserInfo(int32 LocalUserNum, EOS_EpicAccountId EpicAccountId, const FRemoteUserProcessedCallback& Callback)
 {
 	FReadUserInfoCallback* CallbackObj = new FReadUserInfoCallback(AsWeak());
-	CallbackObj->CallbackLambda = [this, LocalUserNum, EpicAccountId](const EOS_UserInfo_QueryUserInfoCallbackInfo* Data)
+	CallbackObj->CallbackLambda = [this, LocalUserNum, EpicAccountId, Callback](const EOS_UserInfo_QueryUserInfoCallbackInfo* Data)
 	{
-		if (Data->ResultCode == EOS_EResult::EOS_Success)
+		const FUniqueNetIdEOSRef EOSId = FUniqueNetIdEOSRegistry::FindChecked(Data->TargetUserId);
+
+		const bool bWasSuccessful = Data->ResultCode == EOS_EResult::EOS_Success;
+		FString ErrorStr;
+
+		if (bWasSuccessful)
 		{
-			const FUniqueNetIdEOSRef EOSId = FUniqueNetIdEOSRegistry::FindChecked(Data->TargetUserId);
-			IAttributeAccessInterfaceRef AttributeAccessRef = UniqueNetIdToAttributeAccessMap[EOSId];
+			IAttributeAccessInterfaceRef AttributeAccessRef = UniqueNetIdToUserRefMap[EOSId];
 			UpdateUserInfo(AttributeAccessRef, Data->LocalUserId, Data->TargetUserId);
 
 			FLocalUserEOS& LocalUser = GetLocalUserChecked(LocalUserNum);
@@ -2787,11 +2890,15 @@ void FUserManagerEOS::ReadUserInfo(int32 LocalUserNum, EOS_EpicAccountId EpicAcc
 				FriendPtr->UpdateInternalAttributes(AttributeAccessRef->GetInternalAttributes());
 			}
 		}
+		else
+		{
+			ErrorStr = LexToString(Data->ResultCode);
+		}
 
 		// We mark this player as processed
-		GetLocalUserChecked(LocalUserNum).OngoingFriendQueryUserInfo.RemoveSwap(EpicAccountId, false);
+		GetLocalUserChecked(LocalUserNum).OngoingQueryUserInfoAccounts.RemoveSwap(EpicAccountId, false);
 
-		ProcessReadFriendsListComplete(LocalUserNum, true, TEXT(""));
+		Callback(bWasSuccessful, EOSId, ErrorStr);
 	};
 
 	EOS_UserInfo_QueryUserInfoOptions Options = { };
@@ -2802,7 +2909,7 @@ void FUserManagerEOS::ReadUserInfo(int32 LocalUserNum, EOS_EpicAccountId EpicAcc
 	EOS_UserInfo_QueryUserInfo(EOSSubsystem->UserInfoHandle, &Options, CallbackObj, CallbackObj->GetCallbackPtr());
 
 	// We mark this player as pending for processing
-	GetLocalUserChecked(LocalUserNum).OngoingFriendQueryUserInfo.Add(EpicAccountId);
+	GetLocalUserChecked(LocalUserNum).OngoingQueryUserInfoAccounts.Add(EpicAccountId);
 }
 
 bool FUserManagerEOS::GetAllUserInfo(int32 LocalUserNum, TArray<TSharedRef<FOnlineUser>>& OutUsers)
@@ -2884,13 +2991,18 @@ bool FUserManagerEOS::QueryUserIdMapping(const FUniqueNetId& UserId, const FStri
 		bool bWasSuccessful = Result == EOS_EResult::EOS_Success;
 		if (bWasSuccessful)
 		{
-			if (!FUniqueNetIdEOSRegistry::Find(Data->TargetUserId))
+			FUniqueNetIdEOSPtr NetId = FUniqueNetIdEOSRegistry::Find(Data->TargetUserId);
+			bool bIsNetIdRegistered = NetId.IsValid() ? UniqueNetIdToUserRefMap.Contains(NetId.ToSharedRef()) : false;
+			if (!bIsNetIdRegistered)
 			{
+				const FRemoteUserProcessedCallback Callback = [this, LocalUserNum, DisplayNameOrEmail, Delegate](bool bWasSuccessful, FUniqueNetIdEOSRef RemotePlayerNetId, const FString& ErrorStr)
+				{
+					Delegate.ExecuteIfBound(bWasSuccessful, *GetLocalUserChecked(LocalUserNum).UniqueNetId, DisplayNameOrEmail, *RemotePlayerNetId, ErrorStr);
+				};
+
 				// Registering the player will also query the presence/user info data
-				AddRemotePlayer(LocalUserNum, Data->TargetUserId);
+				AddRemotePlayer(LocalUserNum, Data->TargetUserId, Callback);
 			}
-
-			Delegate.ExecuteIfBound(true, *GetLocalUserChecked(LocalUserNum).UniqueNetId, DisplayNameOrEmail, *FUniqueNetIdEOSRegistry::FindOrAdd(Data->TargetUserId, nullptr), ErrorString);
 		}
 		else
 		{
@@ -3013,11 +3125,10 @@ bool FUserManagerEOS::QueryExternalIdMappings(const FUniqueNetId& UserId, const
 				{
 					FCStringAnsi::Strncpy(Options.AccountId, TCHAR_TO_UTF8(*StringId), EOS_CONNECT_EXTERNAL_ACCOUNT_ID_MAX_LENGTH + 1);
 					EOS_ProductUserId ProductUserId = EOS_Connect_GetExternalAccountMapping(EOSSubsystem->ConnectHandle, &Options);
-					if (EOS_ProductUserId_IsValid(ProductUserId) == EOS_TRUE)
-					{
-						EOS_EpicAccountId AccountId = EOS_EpicAccountId_FromString(Options.AccountId);
-						UpdateRemotePlayerProductUserId(AccountId, ProductUserId);
-					}
+					
+					// Even if the ProductUserId retrieved above is null, we'll still add it to the registry, as that just means that the EAS account does not have a linked EOS account yet
+					EOS_EpicAccountId AccountId = EOS_EpicAccountId_FromString(Options.AccountId);
+					FUniqueNetIdEOSRegistry::FindOrAdd(AccountId, ProductUserId);
 				}
 			}
 			else
diff --git a/Engine/Plugins/Online/OnlineSubsystemEOS/Source/OnlineSubsystemEOS/Private/UserManagerEOS.h b/Engine/Plugins/Online/OnlineSubsystemEOS/Source/OnlineSubsystemEOS/Private/UserManagerEOS.h
--- a/Engine/Plugins/Online/OnlineSubsystemEOS/Source/OnlineSubsystemEOS/Private/UserManagerEOS.h
+++ b/Engine/Plugins/Online/OnlineSubsystemEOS/Source/OnlineSubsystemEOS/Private/UserManagerEOS.h
@@ -68,6 +68,12 @@ public:
 		TOnlineFriendEOS<FOnlineFriend>(InUserId)
 	{
 	}
+
+	FOnlineFriendEOS(FUniqueNetIdEOSRef InUserId, const TMap<FString, FString>& InUserAttributes) :
+		TOnlineFriendEOS<FOnlineFriend>(InUserId, InUserAttributes)
+	{
+	}
+
 	virtual ~FOnlineFriendEOS() = default;
 };
 
@@ -197,8 +203,27 @@ struct FLocalUserEOS
 	/** Friends information for the local user */
 	FFriendsListEOSPtr FriendsList;
 
-	/** Epic ids for friends with ongoing queries */
-	TArray<EOS_EpicAccountId> OngoingFriendQueryUserInfo;
+	/** Epic ids for users with ongoing queries */
+	TArray<EOS_EpicAccountId> OngoingQueryUserInfoAccounts;
+
+	struct FReadUserInfoResults
+	{
+		/** UniqueNetIds gathered as a result of successful queries */
+		TArray<FUniqueNetIdEOSRef> ProcessedIds;
+
+		/** Whether all queries for this given operation were successful */
+		bool bAllWasSuccessful = true;
+
+		/** Combined error str for all operations */
+		FString AllErrorStr;
+
+		void Reset()
+		{
+			ProcessedIds.Empty();
+			bAllWasSuccessful = true;
+			AllErrorStr.Empty();
+		}
+	} OngoingQueryUserInfoResults;
 
 	/** External ids still in process of resolution for local user */
 	TArray<FString> OngoingPlayerQueryExternalMappings;
@@ -336,10 +361,13 @@ public:
 	bool IsLocalUser(const EOS_ProductUserId& ProductUserId) const;
 
 	typedef TFunction<void(TMap<EOS_ProductUserId, FUniqueNetIdEOSRef> ResolvedUniqueNetIds)> FResolveUniqueNetIdsCallback;
+	typedef TFunction<void(bool bWasSuccessful, TMap<EOS_EpicAccountId, FUniqueNetIdEOSRef> ResolvedUniqueNetIds, FString ErrorStr)> FResolveEpicAccountIdsCallback;
 	typedef TFunction<void(FUniqueNetIdEOSRef ResolvedUniqueNetId)> FResolveUniqueNetIdCallback;
 	bool GetEpicAccountIdFromProductUserId(int32 LocalUserNum, const EOS_ProductUserId& ProductUserId, EOS_EpicAccountId& OutEpicAccountId) const;
 	void ResolveUniqueNetId(int32 LocalUserNum, const EOS_ProductUserId& ProductUserId, const FResolveUniqueNetIdCallback& Callback) const;
 	void ResolveUniqueNetIds(int32 LocalUserNum, const TArray<EOS_ProductUserId>& ProductUserIds, const FResolveUniqueNetIdsCallback& Callback) const;
+	void ResolveUniqueNetId(int32 LocalUserNum, const EOS_EpicAccountId& EpicAccountId, const FResolveUniqueNetIdCallback& Callback);
+	void ResolveUniqueNetIds(int32 LocalUserNum, const TArray<EOS_EpicAccountId>& EpicAccountId, const FResolveEpicAccountIdsCallback& Callback);
 
 	FOnlineUserPtr GetLocalOnlineUser(int32 LocalUserNum) const;
 	FOnlineUserPtr GetOnlineUser(EOS_ProductUserId UserId) const;
@@ -359,6 +387,7 @@ public:
 
 	void FullLoginCallback(int32 LocalUserNum, EOS_EpicAccountId AccountId, EOS_ProductUserId UserId, const FOnlineAccountCredentials& AccountCredentials);
 	void FriendStatusChanged(const EOS_Friends_OnFriendsUpdateInfo* Data);
+	void FriendStatusChangedImpl(EOS_EpicAccountId LocalUserId, EOS_EpicAccountId TargetUserId, EOS_EFriendsStatus PreviousStatus, EOS_EFriendsStatus CurrentStatus);
 	void LoginStatusChanged(const EOS_Auth_LoginStatusChangedCallbackInfo* Data);
 
 	int32 GetDefaultLocalUser() const { return DefaultLocalUser; }
@@ -371,11 +400,14 @@ private:
 	void RemoveLocalUser(int32 LocalUserNum);
 	FLocalUserEOS& AddLocalUser(int32 LocalUserNum, EOS_EpicAccountId EpicAccountId, EOS_ProductUserId UserId, const FOnlineAccountCredentials& AccountCredentials);
 
-	FOnlineFriendEOSRef AddFriend(int32 LocalUserNum, EOS_EpicAccountId EpicAccountId);
-	void AddRemotePlayer(int32 LocalUserNum, EOS_EpicAccountId EpicAccountId);
-	void AddRemotePlayer(int32 LocalUserNum, EOS_EpicAccountId EpicAccountId, IAttributeAccessInterfaceRef AttributeRef);
-	void UpdateRemotePlayerProductUserId(EOS_EpicAccountId AccountId, EOS_ProductUserId UserId);
-	void ReadUserInfo(int32 LocalUserNum, EOS_EpicAccountId EpicAccountId);
+	typedef TFunction<void(bool bWasSuccessful, FUniqueNetIdEOSRef RemotePlayerNetId, const FString& ErrorStr)> FRemoteUserProcessedCallback;
+	void AddRemotePlayer(int32 LocalUserNum, EOS_EpicAccountId EpicAccountId, const FRemoteUserProcessedCallback& Callback);
+	typedef TFunction<void(bool bWasSuccessful, TArray<FUniqueNetIdEOSRef> RemotePlayerNetIds, const FString& ErrorStr)> FRemoteUsersProcessedCallback;
+	void AddRemotePlayers(int32 LocalUserNum, TArray<EOS_EpicAccountId> EpicAccountIds, const FRemoteUsersProcessedCallback& Callback);
+	void ReadUserInfo(int32 LocalUserNum, EOS_EpicAccountId EpicAccountId, const FRemoteUserProcessedCallback& Callback);
+
+	// This method does not call AddRemotePlayer, only registers an already existing remote player as a friend for a local user
+	FOnlineFriendEOSRef AddFriend(int32 LocalUserNum, const FUniqueNetIdEOS& FriendNetId);
 
 	void UpdateUserInfo(IAttributeAccessInterfaceRef AttriubteAccessRef, EOS_EpicAccountId LocalId, EOS_EpicAccountId TargetId);
 	bool IsFriendQueryUserInfoOngoing(int32 LocalUserNum);
@@ -409,7 +441,7 @@ private:
 	// Online user info
 
 	/** Ids mapped to attribute access information for any user */
-	TUniqueNetIdMap<IAttributeAccessInterfaceRef> UniqueNetIdToAttributeAccessMap;
+	TUniqueNetIdMap<FOnlineUserEOSRef> UniqueNetIdToUserRefMap;
 
 	/** Ids mapped to user presence information for any user */
 	TUniqueNetIdMap<FOnlineUserPresenceRef> UniqueNetIdToOnlineUserPresenceMap;
