Refactored GetAuthToken in UserManagerEOS to always make an API call instead of returning a cached value.
Removed all usage of said cached value throughout class.
Refactored equivalent method GetAccessToken to call GetAuthToken.

diff --git a/Engine/Plugins/Online/OnlineSubsystemEOS/Source/OnlineSubsystemEOS/Private/OnlineSubsystemEOSTypes.h b/Engine/Plugins/Online/OnlineSubsystemEOS/Source/OnlineSubsystemEOS/Private/OnlineSubsystemEOSTypes.h
--- a/Engine/Plugins/Online/OnlineSubsystemEOS/Source/OnlineSubsystemEOS/Private/OnlineSubsystemEOSTypes.h
+++ b/Engine/Plugins/Online/OnlineSubsystemEOS/Source/OnlineSubsystemEOS/Private/OnlineSubsystemEOSTypes.h
@@ -239,8 +239,9 @@ class TUserOnlineAccountEOS :
 	public TOnlineUserEOS<BaseClass, IAttributeAccessInterface>
 {
 public:
-	TUserOnlineAccountEOS(const FUniqueNetIdEOSRef& InNetIdRef)
+	TUserOnlineAccountEOS(const FUniqueNetIdEOSRef& InNetIdRef, const FOnlineSubsystemEOS& InSubsystem)
 		: TOnlineUserEOS<BaseClass, IAttributeAccessInterface>(InNetIdRef)
+		, EOSSubsystem(InSubsystem)
 	{
 	}
 
@@ -248,7 +249,13 @@ public:
 	virtual FString GetAccessToken() const override
 	{
 		FString Token;
-		GetAuthAttribute(AUTH_ATTR_ID_TOKEN, Token);
+
+		if (const IOnlineIdentityPtr Identity = EOSSubsystem.GetIdentityInterface())
+		{
+			const int32 LocalUserNum = Identity->GetLocalUserNumFromPlatformUserId(Identity->GetPlatformUserIdFromUniqueNetId(*this->UserIdRef));
+			Token = EOSSubsystem.GetIdentityInterface()->GetAuthToken(LocalUserNum);
+		}
+
 		return Token;
 	}
 
@@ -283,6 +290,7 @@ public:
 protected:
 	/** Additional key/value pair data related to auth */
 	TMap<FString, FString> AdditionalAuthData;
+	const FOnlineSubsystemEOS& EOSSubsystem;
 };
 
 typedef TSharedRef<FOnlineUserPresence> FOnlineUserPresenceRef;
diff --git a/Engine/Plugins/Online/OnlineSubsystemEOS/Source/OnlineSubsystemEOS/Private/UserManagerEOS.cpp b/Engine/Plugins/Online/OnlineSubsystemEOS/Source/OnlineSubsystemEOS/Private/UserManagerEOS.cpp
--- a/Engine/Plugins/Online/OnlineSubsystemEOS/Source/OnlineSubsystemEOS/Private/UserManagerEOS.cpp
+++ b/Engine/Plugins/Online/OnlineSubsystemEOS/Source/OnlineSubsystemEOS/Private/UserManagerEOS.cpp
@@ -766,24 +766,19 @@ void FUserManagerEOS::RefreshConnectLogin(int32 LocalUserNum)
 	const FEOSSettings Settings = UEOSSettings::GetSettings();
 	if (Settings.bUseEAS)
 	{
-		EOS_Auth_Token* AuthToken = nullptr;
-		EOS_Auth_CopyUserAuthTokenOptions CopyOptions = { };
-		CopyOptions.ApiVersion = 1;
-		UE_EOS_CHECK_API_MISMATCH(EOS_AUTH_COPYUSERAUTHTOKEN_API_LATEST, 1);
-
-		EOS_EResult CopyResult = EOS_Auth_CopyUserAuthToken(EOSSubsystem->AuthHandle, &CopyOptions, AccountId, &AuthToken);
-		if (CopyResult == EOS_EResult::EOS_Success)
+		const FString AccessToken = GetAuthToken(LocalUserNum);
+		if (!AccessToken.IsEmpty())
 		{
 			// We update the auth token cached in the user account, along with the user information
 			const FUserOnlineAccountEOSRef UserAccountRef = GetLocalUserChecked(LocalUserNum).UserOnlineAccount.ToSharedRef();
-			UserAccountRef->SetAuthAttribute(AUTH_ATTR_ID_TOKEN, AuthToken->AccessToken);
 			UpdateUserInfo(UserAccountRef, AccountId, AccountId);
 
 			EOS_Connect_Credentials Credentials = { };
 			Credentials.ApiVersion = 1;
 			UE_EOS_CHECK_API_MISMATCH(EOS_CONNECT_CREDENTIALS_API_LATEST, 1);
 			Credentials.Type = EOS_EExternalCredentialType::EOS_ECT_EPIC;
-			Credentials.Token = AuthToken->AccessToken;
+			auto AccessTokenConverter = StringCast<UTF8CHAR>(*AccessToken);
+			Credentials.Token = (const char*)AccessTokenConverter.Get();
 
 			EOS_Connect_LoginOptions Options = { };
 			Options.ApiVersion = 2;
@@ -800,12 +795,10 @@ void FUserManagerEOS::RefreshConnectLogin(int32 LocalUserNum)
 				}
 			};
 			EOS_Connect_Login(EOSSubsystem->ConnectHandle, &Options, CallbackObj, CallbackObj->GetCallbackPtr());
-
-			EOS_Auth_Token_Release(AuthToken);
 		}
 		else
 		{
-			UE_LOG_ONLINE(Error, TEXT("Failed to refresh ConnectLogin(%d) failed with EOS result code (%s)"), LocalUserNum, ANSI_TO_TCHAR(EOS_EResult_ToString(CopyResult)));
+			UE_LOG_ONLINE(Error, TEXT("[FUserManagerEOS::RefreshConnectLogin] AccessToken for user [%d] is empty"), LocalUserNum);
 			Logout(LocalUserNum);
 		}
 	}
@@ -1027,7 +1020,7 @@ FLocalUserEOS& FUserManagerEOS::AddLocalUser(int32 LocalUserNum, EOS_EpicAccount
 	FUniqueNetIdEOSRef UserNetId = FUniqueNetIdEOSRegistry::FindOrAdd(EpicAccountId, UserId);
 	LocalUser.UniqueNetId = UserNetId;
 
-	FUserOnlineAccountEOSRef UserAccountRef(new FUserOnlineAccountEOS(UserNetId));
+	FUserOnlineAccountEOSRef UserAccountRef(new FUserOnlineAccountEOS(UserNetId, *EOSSubsystem));
 	LocalUser.UserOnlineAccount = UserAccountRef;
 
 	LocalUser.FriendsList = MakeShareable(new FFriendsListEOS(LocalUserNum, UserNetId));
@@ -1059,20 +1052,8 @@ FLocalUserEOS& FUserManagerEOS::AddLocalUser(int32 LocalUserNum, EOS_EpicAccount
 	// Once all the fields are set, we start gathering additional information
 	ReadFriendsList(LocalUserNum, FString());
 
-	// Get auth token info
-	EOS_Auth_Token* AuthToken = nullptr;
-	EOS_Auth_CopyUserAuthTokenOptions Options = { };
-	Options.ApiVersion = 1;
-	UE_EOS_CHECK_API_MISMATCH(EOS_AUTH_COPYUSERAUTHTOKEN_API_LATEST, 1);
-
-	EOS_EResult CopyResult = EOS_Auth_CopyUserAuthToken(EOSSubsystem->AuthHandle, &Options, EpicAccountId, &AuthToken);
-	if (CopyResult == EOS_EResult::EOS_Success)
-	{
-		UserAccountRef->SetAuthAttribute(AUTH_ATTR_ID_TOKEN, AuthToken->AccessToken);
-		EOS_Auth_Token_Release(AuthToken);
-
-		UpdateUserInfo(UserAccountRef, EpicAccountId, EpicAccountId);
-	}
+	// Update user info (display name, country, language)
+	UpdateUserInfo(UserAccountRef, EpicAccountId, EpicAccountId);
 
 	return LocalUser;
 }
@@ -1624,15 +1605,26 @@ FString FUserManagerEOS::GetPlayerNickname(const FUniqueNetId& UserId) const
 
 FString FUserManagerEOS::GetAuthToken(int32 LocalUserNum) const
 {
-	FUniqueNetIdPtr UserId = GetUniquePlayerId(LocalUserNum);
-	if (UserId.IsValid())
+	const EOS_EpicAccountId AccountId = GetLocalEpicAccountId(LocalUserNum);
+
+	EOS_Auth_Token* AuthToken = nullptr;
+	EOS_Auth_CopyUserAuthTokenOptions CopyOptions = { };
+	CopyOptions.ApiVersion = 1;
+	UE_EOS_CHECK_API_MISMATCH(EOS_AUTH_COPYUSERAUTHTOKEN_API_LATEST, 1);
+
+	const EOS_EResult CopyResult = EOS_Auth_CopyUserAuthToken(EOSSubsystem->AuthHandle, &CopyOptions, AccountId, &AuthToken);
+	if (CopyResult == EOS_EResult::EOS_Success)
 	{
-		TSharedPtr<FUserOnlineAccount> UserAccount = GetUserAccount(*UserId);
-		if (UserAccount.IsValid())
-		{
-			return UserAccount->GetAccessToken();
-		}
+		const FString AuthTokenStr(UTF8_TO_TCHAR(AuthToken->AccessToken));
+		EOS_Auth_Token_Release(AuthToken);
+
+		return AuthTokenStr;
 	}
+	else
+	{
+		UE_LOG_ONLINE(Verbose, TEXT("[FUserManagerEOS::GetAuthToken] EOS_Auth_CopyUserAuthToken failed with EOS result code (%s) for user (%d)"), ANSI_TO_TCHAR(EOS_EResult_ToString(CopyResult)), LocalUserNum);
+	}
+
 	return FString();
 }
 
diff --git a/Engine/Plugins/Online/OnlineSubsystemEOS/Source/OnlineSubsystemEOS/Private/UserManagerEOS.h b/Engine/Plugins/Online/OnlineSubsystemEOS/Source/OnlineSubsystemEOS/Private/UserManagerEOS.h
--- a/Engine/Plugins/Online/OnlineSubsystemEOS/Source/OnlineSubsystemEOS/Private/UserManagerEOS.h
+++ b/Engine/Plugins/Online/OnlineSubsystemEOS/Source/OnlineSubsystemEOS/Private/UserManagerEOS.h
@@ -47,8 +47,8 @@ class FUserOnlineAccountEOS :
 	public TUserOnlineAccountEOS<FUserOnlineAccount>
 {
 public:
-	FUserOnlineAccountEOS(FUniqueNetIdEOSRef InUserId)
-		: TUserOnlineAccountEOS<FUserOnlineAccount>(InUserId)
+	FUserOnlineAccountEOS(FUniqueNetIdEOSRef InUserId, FOnlineSubsystemEOS& InSubsystem)
+		: TUserOnlineAccountEOS<FUserOnlineAccount>(InUserId, InSubsystem)
 	{
 	}
 	virtual ~FUserOnlineAccountEOS() = default;
