Addition of missing include to fix Android build error

diff --git a/Engine/Plugins/Online/OnlineSubsystemEOS/Source/OnlineSubsystemEOS/Private/OnlineSubsystemEOSTypes.h b/Engine/Plugins/Online/OnlineSubsystemEOS/Source/OnlineSubsystemEOS/Private/OnlineSubsystemEOSTypes.h
--- a/Engine/Plugins/Online/OnlineSubsystemEOS/Source/OnlineSubsystemEOS/Private/OnlineSubsystemEOSTypes.h
+++ b/Engine/Plugins/Online/OnlineSubsystemEOS/Source/OnlineSubsystemEOS/Private/OnlineSubsystemEOSTypes.h
@@ -10,6 +10,7 @@
 #include "EOSSharedTypes.h"
 #include "IPAddress.h"
 #include "OnlineSubsystem.h"
+#include "OnlineSubsystemEOS.h"
 #include "OnlineSubsystemTypes.h"
 #include "OnlineSubsystemEOSPackage.h" // IWYU pragma: keep
 
diff --git a/Engine/Plugins/Online/OnlineSubsystemEOS/Source/OnlineSubsystemEOS/Private/OnlineUserCloudEOS.cpp b/Engine/Plugins/Online/OnlineSubsystemEOS/Source/OnlineSubsystemEOS/Private/OnlineUserCloudEOS.cpp
--- a/Engine/Plugins/Online/OnlineSubsystemEOS/Source/OnlineSubsystemEOS/Private/OnlineUserCloudEOS.cpp
+++ b/Engine/Plugins/Online/OnlineSubsystemEOS/Source/OnlineSubsystemEOS/Private/OnlineUserCloudEOS.cpp
@@ -4,7 +4,6 @@
 #include "OnlineSubsystem.h"
 #include "OnlineSubsystemEOS.h"
 #include "OnlineSubsystemEOSPrivate.h"
-#include "OnlineSubsystemEOSTypes.h"
 #include "UserManagerEOS.h"
 #include "EOSSettings.h"
 
