Fix to Connect auth refresh when running OSS EOS as default platform.

diff --git a/Engine/Plugins/Online/OnlineSubsystemEOS/Source/OnlineSubsystemEOS/Private/UserManagerEOS.cpp b/Engine/Plugins/Online/OnlineSubsystemEOS/Source/OnlineSubsystemEOS/Private/UserManagerEOS.cpp
--- a/Engine/Plugins/Online/OnlineSubsystemEOS/Source/OnlineSubsystemEOS/Private/UserManagerEOS.cpp
+++ b/Engine/Plugins/Online/OnlineSubsystemEOS/Source/OnlineSubsystemEOS/Private/UserManagerEOS.cpp
@@ -764,7 +764,10 @@ void FUserManagerEOS::RefreshConnectLogin(int32 LocalUserNum)
 	}
 
 	const FEOSSettings Settings = UEOSSettings::GetSettings();
-	if (Settings.bUseEAS)
+	// In the case where bIsDefaultOSS is true, FUserManagerEOS::Login will default to using EOS_Auth_Login regardless of the value that bUseEAS is set to
+	// This behaviour will be fixed as part of a wider refactor of FUserManagerEOS::Login
+	const bool bShouldUseEOSAuthToken = EOSSubsystem->bIsDefaultOSS || Settings.bUseEAS;
+	if (bShouldUseEOSAuthToken)
 	{
 		const FString AccessToken = GetAuthToken(LocalUserNum);
 		if (!AccessToken.IsEmpty())
