Refactor to keep lobby details handles alive matching OSS Session objects. Refactored session handles to match new logic.

diff --git a/Engine/Plugins/Online/OnlineSubsystemEOS/Source/OnlineSubsystemEOS/Private/OnlineSessionEOS.cpp b/Engine/Plugins/Online/OnlineSubsystemEOS/Source/OnlineSubsystemEOS/Private/OnlineSessionEOS.cpp
--- a/Engine/Plugins/Online/OnlineSubsystemEOS/Source/OnlineSubsystemEOS/Private/OnlineSessionEOS.cpp
+++ b/Engine/Plugins/Online/OnlineSubsystemEOS/Source/OnlineSubsystemEOS/Private/OnlineSessionEOS.cpp
@@ -363,14 +363,16 @@ FOnlineSessionInfoEOS::FOnlineSessionInfoEOS()
 	: HostAddr(nullptr)
 	, SessionId(FUniqueNetIdString::EmptyId())
 	, SessionHandle(nullptr)
+	, LobbyHandle(nullptr)
 	, bIsFromClone(false)
 {
 }
 
-FOnlineSessionInfoEOS::FOnlineSessionInfoEOS(const FString& InHostIp, FUniqueNetIdStringRef UniqueNetId, EOS_HSessionDetails InSessionHandle)
+FOnlineSessionInfoEOS::FOnlineSessionInfoEOS(const FString& InHostIp, FUniqueNetIdStringRef UniqueNetId, const TSharedPtr<FSessionDetailsEOS>& InSessionHandle, const TSharedPtr<FLobbyDetailsEOS>& InLobbyHandle)
 	: FOnlineSessionInfo()
 	, SessionId(UniqueNetId)
 	, SessionHandle(InSessionHandle)
+	, LobbyHandle(InLobbyHandle)
 	, bIsFromClone(false)
 {
 	if (InHostIp.StartsWith(EOS_CONNECTION_URL_PREFIX, ESearchCase::IgnoreCase))
@@ -387,10 +389,22 @@ FOnlineSessionInfoEOS::FOnlineSessionInfoEOS(const FString& InHostIp, FUniqueNet
 
 FOnlineSessionInfoEOS::~FOnlineSessionInfoEOS()
 {
-	if (SessionHandle != nullptr && !bIsFromClone)
-	{
-		EOS_SessionDetails_Release(SessionHandle);
-	}
+
+}
+
+FOnlineSessionInfoEOS FOnlineSessionInfoEOS::Create(const FString& InHostIp, FUniqueNetIdStringRef UniqueNetId)
+{
+	return FOnlineSessionInfoEOS(InHostIp, UniqueNetId, nullptr, nullptr);
+}
+
+FOnlineSessionInfoEOS FOnlineSessionInfoEOS::Create(const FString& InHostIp, FUniqueNetIdStringRef UniqueNetId, const TSharedPtr<FSessionDetailsEOS>& InSessionHandle)
+{
+	return FOnlineSessionInfoEOS(InHostIp, UniqueNetId, InSessionHandle, nullptr);
+}
+
+FOnlineSessionInfoEOS FOnlineSessionInfoEOS::Create(const FString& InHostIp, FUniqueNetIdStringRef UniqueNetId, const TSharedPtr<FLobbyDetailsEOS>& InLobbyHandle)
+{
+	return FOnlineSessionInfoEOS(InHostIp, UniqueNetId, nullptr, InLobbyHandle);
 }
 
 void FOnlineSessionInfoEOS::InitLAN(FOnlineSubsystemEOS* Subsystem)
@@ -669,10 +683,12 @@ void FOnlineSessionEOS::OnSessionInviteReceived(const EOS_Sessions_SessionInvite
 			UE_EOS_CHECK_API_MISMATCH(EOS_SESSIONS_COPYSESSIONHANDLEBYINVITEID_API_LATEST, 1);
 			const auto InviteIdUTF8 = StringCast<UTF8CHAR>(*InviteId);
 			Options.InviteId = (const char*)InviteIdUTF8.Get();
-			EOS_HSessionDetails SessionDetails = nullptr;
-			EOS_EResult Result = EOS_Sessions_CopySessionHandleByInviteId(EOSSubsystem->SessionsHandle, &Options, &SessionDetails);
+			EOS_HSessionDetails SessionDetailsHandle = nullptr;
+			EOS_EResult Result = EOS_Sessions_CopySessionHandleByInviteId(EOSSubsystem->SessionsHandle, &Options, &SessionDetailsHandle);
 			if (Result == EOS_EResult::EOS_Success)
 			{
+				const TSharedRef<FSessionDetailsEOS> SessionDetails = MakeShared<FSessionDetailsEOS>(SessionDetailsHandle);
+
 				LastInviteSearch = MakeShared<FOnlineSessionSearch>();
 				AddSearchResult(SessionDetails, LastInviteSearch.ToSharedRef());
 				TriggerOnSessionInviteReceivedDelegates(*NetId, *FromNetId, EOSSubsystem->GetAppId(), LastInviteSearch->SearchResults[0]);
@@ -696,10 +712,12 @@ void FOnlineSessionEOS::OnSessionInviteAccepted(const EOS_Sessions_SessionInvite
 			UE_EOS_CHECK_API_MISMATCH(EOS_SESSIONS_COPYSESSIONHANDLEBYINVITEID_API_LATEST, 1);
 			const auto InviteIdUTF8 = StringCast<UTF8CHAR>(*InviteId);
 			Options.InviteId = (const char*)InviteIdUTF8.Get();
-			EOS_HSessionDetails SessionDetails = nullptr;
-			EOS_EResult Result = EOS_Sessions_CopySessionHandleByInviteId(EOSSubsystem->SessionsHandle, &Options, &SessionDetails);
+			EOS_HSessionDetails SessionDetailsHandle = nullptr;
+			EOS_EResult Result = EOS_Sessions_CopySessionHandleByInviteId(EOSSubsystem->SessionsHandle, &Options, &SessionDetailsHandle);
 			if (Result == EOS_EResult::EOS_Success)
 			{
+				const TSharedRef<FSessionDetailsEOS> SessionDetails = MakeShared<FSessionDetailsEOS>(SessionDetailsHandle);
+
 				LastInviteSearch = MakeShared<FOnlineSessionSearch>();
 				AddSearchResult(SessionDetails, LastInviteSearch.ToSharedRef());
 				TriggerOnSessionUserInviteAcceptedDelegates(true, LocalUserNum, ResolvedUniqueNetId, LastInviteSearch->SearchResults[0]);
@@ -722,10 +740,12 @@ void FOnlineSessionEOS::OnJoinSessionAccepted(const EOS_Sessions_JoinSessionAcce
 			Options.ApiVersion = 1;
 			UE_EOS_CHECK_API_MISMATCH(EOS_SESSIONS_COPYSESSIONHANDLEBYUIEVENTID_API_LATEST, 1);
 			Options.UiEventId = UiEventId;
-			EOS_HSessionDetails SessionDetails = nullptr;
-			EOS_EResult Result = EOS_Sessions_CopySessionHandleByUiEventId(EOSSubsystem->SessionsHandle, &Options, &SessionDetails);
+			EOS_HSessionDetails SessionDetailsHandle = nullptr;
+			EOS_EResult Result = EOS_Sessions_CopySessionHandleByUiEventId(EOSSubsystem->SessionsHandle, &Options, &SessionDetailsHandle);
 			if (Result == EOS_EResult::EOS_Success)
 			{
+				const TSharedRef<FSessionDetailsEOS> SessionDetails = MakeShared<FSessionDetailsEOS>(SessionDetailsHandle);
+
 				LastInviteSearch = MakeShared<FOnlineSessionSearch>();
 				AddSearchResult(SessionDetails, LastInviteSearch.ToSharedRef());
 				TriggerOnSessionUserInviteAcceptedDelegates(true, LocalUserNum, ResolvedUniqueNetId, LastInviteSearch->SearchResults[0]);
@@ -844,7 +864,7 @@ void FOnlineSessionEOS::OnLobbyUpdateReceived(const EOS_LobbyId& LobbyId)
 		EOS_EResult CopyLobbyDetailsResult = EOS_Lobby_CopyLobbyDetailsHandle(LobbyHandle, &Options, &LobbyDetailsHandle);
 		if (CopyLobbyDetailsResult == EOS_EResult::EOS_Success)
 		{
-			TSharedRef<FLobbyDetailsEOS> LobbyDetails = MakeShared<FLobbyDetailsEOS>(LobbyDetailsHandle);
+			const TSharedRef<FLobbyDetailsEOS> LobbyDetails = MakeShared<FLobbyDetailsEOS>(LobbyDetailsHandle);
 
 			EOS_LobbyDetails_Info* LobbyDetailsInfo = nullptr;
 			EOS_LobbyDetails_CopyInfoOptions CopyOptions = { };
@@ -854,7 +874,7 @@ void FOnlineSessionEOS::OnLobbyUpdateReceived(const EOS_LobbyId& LobbyId)
 			EOS_EResult CopyInfoResult = EOS_LobbyDetails_CopyInfo(LobbyDetails->LobbyDetailsHandle, &CopyOptions, &LobbyDetailsInfo);
 			if (CopyInfoResult == EOS_EResult::EOS_Success)
 			{
-				CopyLobbyData(LobbyDetails, LobbyDetailsInfo, *Session, [this, SessionName = Session->SessionName, LobbyDetails](bool bWasSuccessful) {
+				CopyLobbyData(LobbyDetails, LobbyDetailsInfo, *Session, [this, SessionName = Session->SessionName](bool bWasSuccessful) {
 					if (bWasSuccessful)
 					{
 						if (FNamedOnlineSession* Session = GetNamedSession(SessionName))
@@ -863,8 +883,6 @@ void FOnlineSessionEOS::OnLobbyUpdateReceived(const EOS_LobbyId& LobbyId)
 						}
 					}
 				});
-
-				EOS_LobbyDetails_Info_Release(LobbyDetailsInfo);
 			}
 			else
 			{
@@ -1053,7 +1071,7 @@ void FOnlineSessionEOS::OnLobbyInviteReceived(const EOS_Lobby_LobbyInviteReceive
 			EOS_EResult Result = EOS_Lobby_CopyLobbyDetailsHandleByInviteId(LobbyHandle, &Options, &LobbyDetailsHandle);
 			if (Result == EOS_EResult::EOS_Success)
 			{
-				TSharedRef<FLobbyDetailsEOS> LobbyDetails = MakeShared<FLobbyDetailsEOS>(LobbyDetailsHandle);
+				const TSharedRef<FLobbyDetailsEOS> LobbyDetails = MakeShared<FLobbyDetailsEOS>(LobbyDetailsHandle);
 
 				LastInviteSearch = MakeShared<FOnlineSessionSearch>();
 				AddLobbySearchResult(LobbyDetails, LastInviteSearch.ToSharedRef(), [this, LocalUserNum, NetId, FromNetId](bool bWasSuccessful)
@@ -1086,7 +1104,7 @@ void FOnlineSessionEOS::OnLobbyInviteAccepted(const EOS_Lobby_LobbyInviteAccepte
 			EOS_EResult Result = EOS_Lobby_CopyLobbyDetailsHandleByInviteId(LobbyHandle, &Options, &LobbyDetailsHandle);
 			if (Result == EOS_EResult::EOS_Success)
 			{
-				TSharedRef<FLobbyDetailsEOS> LobbyDetails = MakeShared<FLobbyDetailsEOS>(LobbyDetailsHandle);
+				const TSharedRef<FLobbyDetailsEOS> LobbyDetails = MakeShared<FLobbyDetailsEOS>(LobbyDetailsHandle);
 
 				LastInviteSearch = MakeShared<FOnlineSessionSearch>();
 				AddLobbySearchResult(LobbyDetails, LastInviteSearch.ToSharedRef(), [this, LocalUserNum, ResolvedUniqueNetId](bool bWasSuccessful)
@@ -1118,7 +1136,7 @@ void FOnlineSessionEOS::OnJoinLobbyAccepted(const EOS_Lobby_JoinLobbyAcceptedCal
 			EOS_EResult Result = EOS_Lobby_CopyLobbyDetailsHandleByUiEventId(LobbyHandle, &Options, &LobbyDetailsHandle);
 			if (Result == EOS_EResult::EOS_Success)
 			{
-				TSharedRef<FLobbyDetailsEOS> LobbyDetails = MakeShared<FLobbyDetailsEOS>(LobbyDetailsHandle);
+				const TSharedRef<FLobbyDetailsEOS> LobbyDetails = MakeShared<FLobbyDetailsEOS>(LobbyDetailsHandle);
 
 				LastInviteSearch = MakeShared<FOnlineSessionSearch>();
 				AddLobbySearchResult(LobbyDetails, LastInviteSearch.ToSharedRef(), [this, LocalUserNum, ResolvedUniqueNetId](bool bWasSuccessful)
@@ -1541,7 +1559,7 @@ uint32 FOnlineSessionEOS::CreateEOSSession(int32 HostingPlayerNum, FNamedOnlineS
 
 	UE_LOG_ONLINE_SESSION(Verbose, TEXT("[FOnlineSessionEOS::CreateEOSSession] The HostAddress used for this session will be %s"), *HostAddr);
 
-	Session->SessionInfo = MakeShareable(new FOnlineSessionInfoEOS(HostAddr, FUniqueNetIdEOSSession::Create(FString()), nullptr));
+	Session->SessionInfo = MakeShared<FOnlineSessionInfoEOS>(FOnlineSessionInfoEOS::Create(HostAddr, FUniqueNetIdEOSSession::Create(FString())));
 
 	FName SessionName = Session->SessionName;
 
@@ -2076,7 +2094,7 @@ uint32 FOnlineSessionEOS::DestroyEOSSession(FNamedOnlineSession* Session, const
 
 	FSessionDestroyOptions Options(TCHAR_TO_UTF8(*Session->SessionName.ToString()));
 	FDestroySessionCallback* CallbackObj = new FDestroySessionCallback(FOnlineSessionEOSWeakPtr(AsShared()));
-	CallbackObj->CallbackLambda = [this, SessionName = Session->SessionName](const EOS_Sessions_DestroySessionCallbackInfo* Data)
+	CallbackObj->CallbackLambda = [this, SessionName = Session->SessionName, CompletionDelegate](const EOS_Sessions_DestroySessionCallbackInfo* Data)
 	{
 		EndSessionAnalytics();
 
@@ -2097,6 +2115,7 @@ uint32 FOnlineSessionEOS::DestroyEOSSession(FNamedOnlineSession* Session, const
 		}
 
 		RemoveNamedSession(SessionName);
+		CompletionDelegate.ExecuteIfBound(SessionName, bWasSuccessful);
 		TriggerOnDestroySessionCompleteDelegates(SessionName, bWasSuccessful);
 	};
 
@@ -2275,7 +2294,7 @@ void FOnlineSessionEOS::AddLobbySearchAttribute(EOS_HLobbySearch LobbySearchHand
 	}
 }
 
-void FOnlineSessionEOS::CopySearchResult(EOS_HSessionDetails SessionHandle, EOS_SessionDetails_Info* SessionInfo, FOnlineSession& OutSession)
+void FOnlineSessionEOS::CopySearchResult(const FSessionDetailsEOS& SessionHandle, EOS_SessionDetails_Info* SessionInfo, FOnlineSession& OutSession)
 {
 	OutSession.NumOpenPrivateConnections = SessionInfo->NumOpenPublicConnections;
 	OutSession.SessionSettings.NumPrivateConnections = SessionInfo->Settings->NumPublicConnections;
@@ -2300,12 +2319,12 @@ void FOnlineSessionEOS::CopySearchResult(EOS_HSessionDetails SessionHandle, EOS_
 	CopyAttributes(SessionHandle, OutSession);
 }
 
-void FOnlineSessionEOS::CopyAttributes(EOS_HSessionDetails SessionHandle, FOnlineSession& OutSession)
+void FOnlineSessionEOS::CopyAttributes(const FSessionDetailsEOS& SessionHandle, FOnlineSession& OutSession)
 {
 	EOS_SessionDetails_GetSessionAttributeCountOptions CountOptions = { };
 	CountOptions.ApiVersion = 1;
 	UE_EOS_CHECK_API_MISMATCH(EOS_SESSIONDETAILS_GETSESSIONATTRIBUTECOUNT_API_LATEST, 1);
-	int32 Count = EOS_SessionDetails_GetSessionAttributeCount(SessionHandle, &CountOptions);
+	int32 Count = EOS_SessionDetails_GetSessionAttributeCount(SessionHandle.SessionDetailsHandle, &CountOptions);
 
 	for (int32 Index = 0; Index < Count; Index++)
 	{
@@ -2315,7 +2334,7 @@ void FOnlineSessionEOS::CopyAttributes(EOS_HSessionDetails SessionHandle, FOnlin
 		AttrOptions.AttrIndex = Index;
 
 		EOS_SessionDetails_Attribute* Attribute = NULL;
-		EOS_EResult ResultCode = EOS_SessionDetails_CopySessionAttributeByIndex(SessionHandle, &AttrOptions, &Attribute);
+		EOS_EResult ResultCode = EOS_SessionDetails_CopySessionAttributeByIndex(SessionHandle.SessionDetailsHandle, &AttrOptions, &Attribute);
 		if (ResultCode == EOS_EResult::EOS_Success)
 		{
 			FString Key = Attribute->Data->Key;
@@ -2385,21 +2404,21 @@ void FOnlineSessionEOS::CopyAttributes(EOS_HSessionDetails SessionHandle, FOnlin
 	}
 }
 
-void FOnlineSessionEOS::AddSearchResult(EOS_HSessionDetails SessionHandle, const TSharedRef<FOnlineSessionSearch>& SearchSettings)
+void FOnlineSessionEOS::AddSearchResult(const TSharedRef<FSessionDetailsEOS>& SessionHandle, const TSharedRef<FOnlineSessionSearch>& SearchSettings)
 {
 	EOS_SessionDetails_Info* SessionInfo = nullptr;
 	EOS_SessionDetails_CopyInfoOptions CopyOptions = { };
 	CopyOptions.ApiVersion = 1;
 	UE_EOS_CHECK_API_MISMATCH(EOS_SESSIONDETAILS_COPYINFO_API_LATEST, 1);
-	EOS_EResult CopyResult = EOS_SessionDetails_CopyInfo(SessionHandle, &CopyOptions, &SessionInfo);
+	EOS_EResult CopyResult = EOS_SessionDetails_CopyInfo(SessionHandle->SessionDetailsHandle, &CopyOptions, &SessionInfo);
 	if (CopyResult == EOS_EResult::EOS_Success)
 	{
 		int32 Position = SearchSettings->SearchResults.AddZeroed();
 		FOnlineSessionSearchResult& SearchResult = SearchSettings->SearchResults[Position];
 		// This will set the host address and port
-		SearchResult.Session.SessionInfo = MakeShareable(new FOnlineSessionInfoEOS(SessionInfo->HostAddress, FUniqueNetIdEOSSession::Create(SessionInfo->SessionId), SessionHandle));
+		SearchResult.Session.SessionInfo = MakeShared<FOnlineSessionInfoEOS>(FOnlineSessionInfoEOS::Create(SessionInfo->HostAddress, FUniqueNetIdEOSSession::Create(SessionInfo->SessionId), SessionHandle));
 
-		CopySearchResult(SessionHandle, SessionInfo, SearchResult.Session);
+		CopySearchResult(*SessionHandle, SessionInfo, SearchResult.Session);
 
 		EOS_SessionDetails_Info_Release(SessionInfo);
 	}
@@ -2466,12 +2485,14 @@ uint32 FOnlineSessionEOS::FindEOSSession(int32 SearchingPlayerNum, const TShared
 			UE_EOS_CHECK_API_MISMATCH(EOS_SESSIONSEARCH_COPYSEARCHRESULTBYINDEX_API_LATEST, 1);
 			for (int32 Index = 0; Index < NumSearchResults; Index++)
 			{
-				EOS_HSessionDetails SessionHandle = nullptr;
+				EOS_HSessionDetails SessionDetailsHandle = nullptr;
 				IndexOptions.SessionIndex = Index;
-				EOS_EResult Result = EOS_SessionSearch_CopySearchResultByIndex(CurrentSearchHandle->SearchHandle, &IndexOptions, &SessionHandle);
+				EOS_EResult Result = EOS_SessionSearch_CopySearchResultByIndex(CurrentSearchHandle->SearchHandle, &IndexOptions, &SessionDetailsHandle);
 				if (Result == EOS_EResult::EOS_Success)
 				{
-					AddSearchResult(SessionHandle, SearchSettings);
+					const TSharedRef<FSessionDetailsEOS> SessionDetails = MakeShared<FSessionDetailsEOS>(SessionDetailsHandle);
+
+					AddSearchResult(SessionDetails, SearchSettings);
 				}
 			}
 			SearchSettings->SearchState = EOnlineAsyncTaskState::Done;
@@ -2547,12 +2568,14 @@ void FOnlineSessionEOS::FindEOSSessionById(int32 LocalUserNum, const FUniqueNetI
 			UE_EOS_CHECK_API_MISMATCH(EOS_SESSIONSEARCH_COPYSEARCHRESULTBYINDEX_API_LATEST, 1);
 			for (int32 Index = 0; Index < NumSearchResults; Index++)
 			{
-				EOS_HSessionDetails SessionHandle = nullptr;
+				EOS_HSessionDetails SessionDetailsHandle = nullptr;
 				IndexOptions.SessionIndex = Index;
-				EOS_EResult Result = EOS_SessionSearch_CopySearchResultByIndex(CurrentSearchHandle->SearchHandle, &IndexOptions, &SessionHandle);
+				EOS_EResult Result = EOS_SessionSearch_CopySearchResultByIndex(CurrentSearchHandle->SearchHandle, &IndexOptions, &SessionDetailsHandle);
 				if (Result == EOS_EResult::EOS_Success)
 				{
-					AddSearchResult(SessionHandle, LocalSessionSearch);
+					const TSharedRef<FSessionDetailsEOS> SessionDetails = MakeShared<FSessionDetailsEOS>(SessionDetailsHandle);
+
+					AddSearchResult(SessionDetails, LocalSessionSearch);
 				}
 			}
 			LocalSessionSearch->SearchState = EOnlineAsyncTaskState::Done;
@@ -2795,7 +2818,7 @@ uint32 FOnlineSessionEOS::JoinEOSSession(int32 PlayerNum, FNamedOnlineSession* S
 
 	FJoinSessionOptions Options(TCHAR_TO_UTF8(*Session->SessionName.ToString()));
 	Options.LocalUserId = ProductUserId;
-	Options.SessionHandle = EOSSessionInfo->SessionHandle;
+	Options.SessionHandle = EOSSessionInfo->SessionHandle->SessionDetailsHandle;
 	EOS_Sessions_JoinSession(EOSSubsystem->SessionsHandle, &Options, CallbackObj, CallbackObj->GetCallbackPtr());
 
 	return ONLINE_IO_PENDING;
@@ -3961,7 +3984,7 @@ uint32 FOnlineSessionEOS::CreateLobbySession(int32 HostingPlayerNum, FNamedOnlin
 				FInternetAddrEOS TempAddr(LexToString(LocalProductUserId), SessionName.ToString(), FURL::UrlConfig.DefaultPort);
 				FString HostAddr = TempAddr.ToString(true);
 
-				Session->SessionInfo = MakeShareable(new FOnlineSessionInfoEOS(HostAddr, FUniqueNetIdEOSLobby::Create(UTF8_TO_TCHAR(Data->LobbyId)), nullptr));
+				Session->SessionInfo = MakeShared<FOnlineSessionInfoEOS>(FOnlineSessionInfoEOS::Create(HostAddr, FUniqueNetIdEOSLobby::Create(UTF8_TO_TCHAR(Data->LobbyId))));
 
 #if WITH_EOS_RTC
 				if (FEOSVoiceChatUser* VoiceChatUser = static_cast<FEOSVoiceChatUser*>(EOSSubsystem->GetEOSVoiceChatUserInterface(*LocalUserNetId)))
@@ -4083,6 +4106,7 @@ uint32 FOnlineSessionEOS::JoinLobbySession(int32 PlayerNum, FNamedOnlineSession*
 			EOSSessionInfo->HostAddr = SearchSessionInfo->HostAddr;
 			EOSSessionInfo->EOSAddress = SearchSessionInfo->EOSAddress;
 			EOSSessionInfo->SessionHandle = SearchSessionInfo->SessionHandle;
+			EOSSessionInfo->LobbyHandle = SearchSessionInfo->LobbyHandle;
 			EOSSessionInfo->SessionId = SearchSessionInfo->SessionId;
 			EOSSessionInfo->bIsFromClone = SearchSessionInfo->bIsFromClone;
 
@@ -4094,9 +4118,7 @@ uint32 FOnlineSessionEOS::JoinLobbySession(int32 PlayerNum, FNamedOnlineSession*
 			UE_EOS_CHECK_API_MISMATCH(EOS_LOBBY_JOINLOBBY_API_LATEST, 4);
 			JoinLobbyOptions.LocalUserId = EOSSubsystem->UserManager->GetLocalProductUserId(PlayerNum);
 			JoinLobbyOptions.bPresenceEnabled = Session->SessionSettings.bUsesPresence;
-
- 			TSharedRef<FLobbyDetailsEOS> LobbyDetails = LobbySearchResultsCache[Session->SessionInfo->GetSessionId().ToString()];
- 			JoinLobbyOptions.LobbyDetailsHandle = LobbyDetails->LobbyDetailsHandle;
+ 			JoinLobbyOptions.LobbyDetailsHandle = EOSSessionInfo->LobbyHandle->LobbyDetailsHandle;
 
 			const FName SessionName = Session->SessionName;
 			const FUniqueNetIdEOSPtr LocalUserNetId = EOSSubsystem->UserManager->GetLocalUniqueNetIdEOS(PlayerNum);
@@ -4140,9 +4162,6 @@ uint32 FOnlineSessionEOS::JoinLobbySession(int32 PlayerNum, FNamedOnlineSession*
 					UE_LOG_ONLINE_SESSION(Warning, TEXT("[FOnlineSessionEOS::JoinLobbySession] Unable to find session %s"), *SessionName.ToString());
 					TriggerOnJoinSessionCompleteDelegates(SessionName, EOnJoinSessionCompleteResult::SessionDoesNotExist);
 				}
-
-				// Now that we have finished the join process, we can clear the cache
-				LobbySearchResultsCache.Reset();
 			};
 
 			EOS_Lobby_JoinLobby(LobbyHandle, &JoinLobbyOptions, CallbackObj, CallbackObj->GetCallbackPtr());
@@ -4472,6 +4491,7 @@ uint32 FOnlineSessionEOS::DestroyLobbySession(int32 LocalUserNum, FNamedOnlineSe
 			else
 			{
 				UE_LOG_ONLINE_SESSION(Warning, TEXT("[FOnlineSessionEOS::DestroyLobbySession] Unable to find session %s"), *SessionName.ToString());
+				CompletionDelegate.ExecuteIfBound(SessionName, false);
 				TriggerOnDestroySessionCompleteDelegates(SessionName, false);
 			}
 		};
@@ -4533,9 +4553,6 @@ uint32 FOnlineSessionEOS::FindLobbySession(int32 SearchingPlayerNum, const TShar
 
 void FOnlineSessionEOS::StartLobbySearch(int32 SearchingPlayerNum, EOS_HLobbySearch LobbySearchHandle, const TSharedRef<FOnlineSessionSearch>& SearchSettings, const FOnSingleSessionResultCompleteDelegate& CompletionDelegate)
 {
-	// When starting a new search, we'll reset the cache
-	LobbySearchResultsCache.Reset();
-
 	SessionSearchStartInSeconds = FPlatformTime::Seconds();
 
 	EOS_LobbySearch_FindOptions FindOptions = { 0 };
@@ -4575,8 +4592,8 @@ void FOnlineSessionEOS::StartLobbySearch(int32 SearchingPlayerNum, EOS_HLobbySea
 					if (Result == EOS_EResult::EOS_Success)
 					{
 						UE_LOG_ONLINE_SESSION(Verbose, TEXT("[FOnlineSessionEOS::StartLobbySearch::FLobbySearchFindCallback] LobbySearch_CopySearchResultByIndex was successful."));
-						TSharedRef<FLobbyDetailsEOS> LobbyDetails = MakeShared<FLobbyDetailsEOS>(LobbyDetailsHandle);
-						PendingLobbySearchResults.Add(LobbyDetails);
+						const TSharedRef<FLobbyDetailsEOS> LobbyDetails = MakeShared<FLobbyDetailsEOS>(LobbyDetailsHandle);
+						LobbySearchResultsPendingIdResolution.Add(LobbyDetails);
 					}
 					else
 					{
@@ -4585,14 +4602,14 @@ void FOnlineSessionEOS::StartLobbySearch(int32 SearchingPlayerNum, EOS_HLobbySea
 				}
 
 				// Make a copy to iterate over, as the AddLobbySearchResult delegate removes entries.
-				const TArray<TSharedRef<FLobbyDetailsEOS>> PendingLobbySearchResultsCopy = PendingLobbySearchResults;
-				for (const TSharedRef<FLobbyDetailsEOS>& LobbyDetails : PendingLobbySearchResultsCopy)
+				const TArray<TSharedRef<FLobbyDetailsEOS>> LobbySearchResultsPendingIdResolutionCopy = LobbySearchResultsPendingIdResolution;
+				for (const TSharedRef<FLobbyDetailsEOS>& LobbyDetails : LobbySearchResultsPendingIdResolutionCopy)
 				{
 					AddLobbySearchResult(LobbyDetails, SearchSettings, [this, LobbyDetails, CompletionDelegate, SearchingPlayerNum, SearchSettings](bool bWasSuccessful)
 					{
-						PendingLobbySearchResults.Remove(LobbyDetails);
+						LobbySearchResultsPendingIdResolution.Remove(LobbyDetails);
 
-						if (PendingLobbySearchResults.IsEmpty())
+						if (LobbySearchResultsPendingIdResolution.IsEmpty())
 						{
 							// If we fail to copy the lobby data, we won't add a new search result, so we'll return an empty one
 							CompletionDelegate.ExecuteIfBound(SearchingPlayerNum, bWasSuccessful, bWasSuccessful ? SearchSettings->SearchResults.Last() : FOnlineSessionSearchResult());
@@ -4600,7 +4617,7 @@ void FOnlineSessionEOS::StartLobbySearch(int32 SearchingPlayerNum, EOS_HLobbySea
 					});
 				}
 
-				if (PendingLobbySearchResults.IsEmpty())
+				if (LobbySearchResultsPendingIdResolution.IsEmpty())
 				{
 					CompletionDelegate.ExecuteIfBound(SearchingPlayerNum, true, SearchSettings->SearchResults.Last());
 				}
@@ -4646,10 +4663,9 @@ void FOnlineSessionEOS::AddLobbySearchResult(const TSharedRef<FLobbyDetailsEOS>&
 		FInternetAddrEOS TempAddr(LexToString(LobbyDetailsInfo->LobbyOwnerUserId), NetDriverName.ToString(), GetTypeHash(NetDriverName.ToString()));
 		FString HostAddr = TempAddr.ToString(true);
 
-		SearchResult.Session.SessionInfo = MakeShareable(new FOnlineSessionInfoEOS(HostAddr, FUniqueNetIdEOSLobby::Create(UTF8_TO_TCHAR(LobbyDetailsInfo->LobbyId)), nullptr));
+		SearchResult.Session.SessionInfo = MakeShared<FOnlineSessionInfoEOS>(FOnlineSessionInfoEOS::Create(HostAddr, FUniqueNetIdEOSLobby::Create(UTF8_TO_TCHAR(LobbyDetailsInfo->LobbyId)), LobbyDetails));
 
 		// We copy the lobby data and settings
-		LobbySearchResultsCache.Add(FString(UTF8_TO_TCHAR(LobbyDetailsInfo->LobbyId)), LobbyDetails);
 		CopyLobbyData(LobbyDetails, LobbyDetailsInfo, SearchResult.Session, Callback);
 
 		EOS_LobbyDetails_Info_Release(LobbyDetailsInfo);
@@ -4670,6 +4686,8 @@ void FOnlineSessionEOS::AddLobbySearchResult(const TSharedRef<FLobbyDetailsEOS>&
 
 void FOnlineSessionEOS::CopyLobbyData(const TSharedRef<FLobbyDetailsEOS>& LobbyDetails, EOS_LobbyDetails_Info* LobbyDetailsInfo, FOnlineSession& OutSession, const FOnCopyLobbyDataCompleteCallback& Callback)
 {
+	// This method launches an asynchronous operation, so we'll pass the details handle as a shared ref to make sure it stays alive
+
 	OutSession.SessionSettings.bUseLobbiesIfAvailable = true;
 	OutSession.SessionSettings.bIsLANMatch = false;
 	OutSession.SessionSettings.Set(SETTING_HOST_MIGRATION, LobbyDetailsInfo->bAllowHostMigration, EOnlineDataAdvertisementType::DontAdvertise);
@@ -4699,7 +4717,7 @@ void FOnlineSessionEOS::CopyLobbyData(const TSharedRef<FLobbyDetailsEOS>& LobbyD
 	OutSession.SessionSettings.bAllowInvites = (bool)LobbyDetailsInfo->bAllowInvites;
 
 	// We copy the settings related to lobby attributes
-	CopyLobbyAttributes(LobbyDetails, OutSession);
+	CopyLobbyAttributes(*LobbyDetails, OutSession);
 
 	// Then we copy the settings for all lobby members
 	EOS_LobbyDetails_GetMemberCountOptions CountOptions = { };
@@ -4725,7 +4743,6 @@ void FOnlineSessionEOS::CopyLobbyData(const TSharedRef<FLobbyDetailsEOS>& LobbyD
 	{
 		EOSSubsystem->UserManager->ResolveUniqueNetIds(EOSSubsystem->UserManager->GetDefaultLocalUser(), TargetUserIds, [this, LobbyDetails, LobbyId = FUniqueNetIdEOSLobby::Create(UTF8_TO_TCHAR(LobbyDetailsInfo->LobbyId)), OriginalCallback = Callback](TMap<EOS_ProductUserId, FUniqueNetIdEOSRef> ResolvedUniqueNetIds)
 			{
-				
 				FOnlineSession* Session = GetOnlineSessionFromLobbyId(*LobbyId);
 				if (Session)
 				{
@@ -4764,14 +4781,14 @@ void FOnlineSessionEOS::CopyLobbyData(const TSharedRef<FLobbyDetailsEOS>& LobbyD
 	}
 }
 
-void FOnlineSessionEOS::CopyLobbyAttributes(const TSharedRef<FLobbyDetailsEOS>& LobbyDetails, FOnlineSession& OutSession)
+void FOnlineSessionEOS::CopyLobbyAttributes(const FLobbyDetailsEOS& LobbyDetails, FOnlineSession& OutSession)
 {
 	// In this method we are updating/adding attributes, but not removing
 
 	EOS_LobbyDetails_GetAttributeCountOptions CountOptions = { };
 	CountOptions.ApiVersion = 1;
 	UE_EOS_CHECK_API_MISMATCH(EOS_LOBBYDETAILS_GETATTRIBUTECOUNT_API_LATEST, 1);
-	int32 Count = EOS_LobbyDetails_GetAttributeCount(LobbyDetails->LobbyDetailsHandle, &CountOptions);
+	int32 Count = EOS_LobbyDetails_GetAttributeCount(LobbyDetails.LobbyDetailsHandle, &CountOptions);
 
 	for (int32 Index = 0; Index < Count; Index++)
 	{
@@ -4781,7 +4798,7 @@ void FOnlineSessionEOS::CopyLobbyAttributes(const TSharedRef<FLobbyDetailsEOS>&
 		AttrOptions.AttrIndex = Index;
 
 		EOS_Lobby_Attribute* Attribute = NULL;
-		EOS_EResult ResultCode = EOS_LobbyDetails_CopyAttributeByIndex(LobbyDetails->LobbyDetailsHandle, &AttrOptions, &Attribute);
+		EOS_EResult ResultCode = EOS_LobbyDetails_CopyAttributeByIndex(LobbyDetails.LobbyDetailsHandle, &AttrOptions, &Attribute);
 		if (ResultCode == EOS_EResult::EOS_Success)
 		{
 			FString Key = UTF8_TO_TCHAR(Attribute->Data->Key);
diff --git a/Engine/Plugins/Online/OnlineSubsystemEOS/Source/OnlineSubsystemEOS/Private/OnlineSessionEOS.h b/Engine/Plugins/Online/OnlineSubsystemEOS/Source/OnlineSubsystemEOS/Private/OnlineSessionEOS.h
--- a/Engine/Plugins/Online/OnlineSubsystemEOS/Source/OnlineSubsystemEOS/Private/OnlineSessionEOS.h
+++ b/Engine/Plugins/Online/OnlineSubsystemEOS/Source/OnlineSubsystemEOS/Private/OnlineSessionEOS.h
@@ -34,21 +34,6 @@ struct FSessionSearchEOS
 	}
 };
 
-struct FLobbyDetailsEOS : FNoncopyable
-{
-	EOS_HLobbyDetails LobbyDetailsHandle;
-
-	FLobbyDetailsEOS(EOS_HLobbyDetails InLobbyDetailsHandle)
-		: LobbyDetailsHandle(InLobbyDetailsHandle)
-	{
-	}
-
-	virtual ~FLobbyDetailsEOS()
-	{
-		EOS_LobbyDetails_Release(LobbyDetailsHandle);
-	}
-};
-
 /**
  * Interface for interacting with EOS sessions
  */
@@ -212,8 +197,7 @@ private:
 	// EOS Lobbies
 
 	EOS_HLobby LobbyHandle;
-	TArray<TSharedRef<FLobbyDetailsEOS>> PendingLobbySearchResults;
-	TMap<FString, TSharedRef<FLobbyDetailsEOS>> LobbySearchResultsCache;
+	TArray<TSharedRef<FLobbyDetailsEOS>> LobbySearchResultsPendingIdResolution;
 
 	// Lobby session callbacks and methods
 	FCallbackBase* LobbyCreatedCallback;
@@ -267,7 +251,7 @@ private:
 	// Methods to update an OSS Lobby from an API Lobby
 	typedef TFunction<void(bool bWasSuccessful)> FOnCopyLobbyDataCompleteCallback;
 	void CopyLobbyData(const TSharedRef<FLobbyDetailsEOS>& LobbyDetails, EOS_LobbyDetails_Info* LobbyDetailsInfo, FOnlineSession& OutSession, const FOnCopyLobbyDataCompleteCallback& Callback);
-	void CopyLobbyAttributes(const TSharedRef<FLobbyDetailsEOS>& LobbyDetails, FOnlineSession& OutSession);
+	void CopyLobbyAttributes(const FLobbyDetailsEOS& LobbyDetails, FOnlineSession& OutSession);
 	void CopyLobbyMemberAttributes(const FLobbyDetailsEOS& LobbyDetails, const EOS_ProductUserId& TargetUserId, FSessionSettings& OutSessionSettings);
 
 	// Lobby search
@@ -307,10 +291,10 @@ private:
 	void BeginSessionAnalytics(FNamedOnlineSession* Session);
 	void EndSessionAnalytics();
 
-	void AddSearchResult(EOS_HSessionDetails SessionHandle, const TSharedRef<FOnlineSessionSearch>& SearchSettings);
+	void AddSearchResult(const TSharedRef<FSessionDetailsEOS>& SessionHandle, const TSharedRef<FOnlineSessionSearch>& SearchSettings);
 	void AddSearchAttribute(EOS_HSessionSearch SearchHandle, const EOS_Sessions_AttributeData* Attribute, EOS_EOnlineComparisonOp ComparisonOp);
-	void CopySearchResult(EOS_HSessionDetails SessionHandle, EOS_SessionDetails_Info* SessionInfo, FOnlineSession& SessionSettings);
-	void CopyAttributes(EOS_HSessionDetails SessionHandle, FOnlineSession& OutSession);
+	void CopySearchResult(const FSessionDetailsEOS& SessionHandle, EOS_SessionDetails_Info* SessionInfo, FOnlineSession& SessionSettings);
+	void CopyAttributes(const FSessionDetailsEOS& SessionHandle, FOnlineSession& OutSession);
 
 	void SetPermissionLevel(EOS_HSessionModification SessionModHandle, FNamedOnlineSession* Session);
 	void SetMaxPlayers(EOS_HSessionModification SessionModHandle, FNamedOnlineSession* Session);
diff --git a/Engine/Plugins/Online/OnlineSubsystemEOS/Source/OnlineSubsystemEOS/Private/OnlineSubsystemEOSTypes.h b/Engine/Plugins/Online/OnlineSubsystemEOS/Source/OnlineSubsystemEOS/Private/OnlineSubsystemEOSTypes.h
--- a/Engine/Plugins/Online/OnlineSubsystemEOS/Source/OnlineSubsystemEOS/Private/OnlineSubsystemEOSTypes.h
+++ b/Engine/Plugins/Online/OnlineSubsystemEOS/Source/OnlineSubsystemEOS/Private/OnlineSubsystemEOSTypes.h
@@ -553,6 +553,36 @@ private:
 
 #include "eos_sessions_types.h"
 
+struct FSessionDetailsEOS : FNoncopyable
+{
+	EOS_HSessionDetails SessionDetailsHandle;
+
+	FSessionDetailsEOS(EOS_HSessionDetails InSessionDetailsHandle)
+		: SessionDetailsHandle(InSessionDetailsHandle)
+	{
+	}
+
+	virtual ~FSessionDetailsEOS()
+	{
+		EOS_SessionDetails_Release(SessionDetailsHandle);
+	}
+};
+
+struct FLobbyDetailsEOS : FNoncopyable
+{
+	EOS_HLobbyDetails LobbyDetailsHandle;
+
+	FLobbyDetailsEOS(EOS_HLobbyDetails InLobbyDetailsHandle)
+		: LobbyDetailsHandle(InLobbyDetailsHandle)
+	{
+	}
+
+	virtual ~FLobbyDetailsEOS()
+	{
+		EOS_LobbyDetails_Release(LobbyDetailsHandle);
+	}
+};
+
 /**
  * Implementation of session information
  */
@@ -575,11 +605,16 @@ PACKAGE_SCOPE:
 		, HostAddr(Src.HostAddr)
 		, SessionId(Src.SessionId)
 		, SessionHandle(Src.SessionHandle)
+		, LobbyHandle(Src.LobbyHandle)
 		, bIsFromClone(true)
 	{
 	}
 
-	FOnlineSessionInfoEOS(const FString& InHostIp, FUniqueNetIdStringRef UniqueNetId, EOS_HSessionDetails InSessionHandle);
+	FOnlineSessionInfoEOS(const FString& InHostIp, FUniqueNetIdStringRef UniqueNetId, const TSharedPtr<FSessionDetailsEOS>& InSessionHandle, const TSharedPtr<FLobbyDetailsEOS>& InLobbyHandle);
+
+	static FOnlineSessionInfoEOS Create(const FString& InHostIp, FUniqueNetIdStringRef UniqueNetId);
+	static FOnlineSessionInfoEOS Create(const FString& InHostIp, FUniqueNetIdStringRef UniqueNetId, const TSharedPtr<FSessionDetailsEOS>& InSessionHandle);
+	static FOnlineSessionInfoEOS Create(const FString& InHostIp, FUniqueNetIdStringRef UniqueNetId, const TSharedPtr<FLobbyDetailsEOS>& InLobbyHandle);
 
 	/**
 	 * Initialize LAN session
@@ -591,8 +626,10 @@ PACKAGE_SCOPE:
 	TSharedPtr<class FInternetAddr> HostAddr;
 	/** Unique Id for this session */
 	FUniqueNetIdStringRef SessionId;
-	/** EOS session handle. Note: this needs to be released by the SDK */
-	EOS_HSessionDetails SessionHandle;
+	/** EOS session handle. The same handle can be shared between a local session and a search result. The struct type will call the release API automatically upon destruction */
+	TSharedPtr<FSessionDetailsEOS> SessionHandle;
+	/** EOS lobby handle. The same handle can be shared between a local session and a search result. The struct type will call the release API automatically upon destruction */
+	TSharedPtr<FLobbyDetailsEOS> LobbyHandle;
 	/** Whether we should delete this handle or not */
 	bool bIsFromClone;
 
