- Add ability to override EpicSandboxId command line arg.

diff --git a/Engine/Plugins/Online/OnlineSubsystemEOS/Source/OnlineSubsystemEOS/Private/EOSSettings.cpp b/Engine/Plugins/Online/OnlineSubsystemEOS/Source/OnlineSubsystemEOS/Private/EOSSettings.cpp
--- a/Engine/Plugins/Online/OnlineSubsystemEOS/Source/OnlineSubsystemEOS/Private/EOSSettings.cpp
+++ b/Engine/Plugins/Online/OnlineSubsystemEOS/Source/OnlineSubsystemEOS/Private/EOSSettings.cpp
@@ -206,9 +206,13 @@ bool UEOSSettings::GetSelectedArtifactSettings(FEOSArtifactSettings& OutSettings
 	// Prefer -EOSArtifactNameOverride over previous.
 	FParse::Value(FCommandLine::Get(), TEXT("EOSArtifactNameOverride="), ArtifactName);
 
-	// Get the sandbox id, which only comes in via EGS. If present, grab the settings where both match.
 	FString SandboxId;
-	if (FParse::Value(FCommandLine::Get(), TEXT("EpicSandboxId="), SandboxId))
+	// Get the -epicsandboxid argument. This generally comes from EGS.
+	bool bHasSandboxId = FParse::Value(FCommandLine::Get(), TEXT("EpicSandboxId="), SandboxId);
+	// Prefer -EpicSandboxIdOverride over previous.
+	bHasSandboxId |= FParse::Value(FCommandLine::Get(), TEXT("EpicSandboxIdOverride="), SandboxId);
+	// If present, grab the settings where both match
+	if (bHasSandboxId)
 	{
 		if (GetArtifactSettings(ArtifactName, SandboxId, OutSettings))
 		{
