Added refactor to prevent lobby member information being read on lobby search results. We'll only do it when we are part of the lobby (lobby updates).
No changes needed in OnlineServices.

diff --git a/Engine/Plugins/Online/OnlineSubsystemEOS/Source/OnlineSubsystemEOS/Private/OnlineSessionEOS.cpp b/Engine/Plugins/Online/OnlineSubsystemEOS/Source/OnlineSubsystemEOS/Private/OnlineSessionEOS.cpp
--- a/Engine/Plugins/Online/OnlineSubsystemEOS/Source/OnlineSubsystemEOS/Private/OnlineSessionEOS.cpp
+++ b/Engine/Plugins/Online/OnlineSubsystemEOS/Source/OnlineSubsystemEOS/Private/OnlineSessionEOS.cpp
@@ -874,7 +874,8 @@ void FOnlineSessionEOS::OnLobbyUpdateReceived(const EOS_LobbyId& LobbyId)
 			EOS_EResult CopyInfoResult = EOS_LobbyDetails_CopyInfo(LobbyDetails->LobbyDetailsHandle, &CopyOptions, &LobbyDetailsInfo);
 			if (CopyInfoResult == EOS_EResult::EOS_Success)
 			{
-				CopyLobbyData(LobbyDetails, LobbyDetailsInfo, *Session, [this, SessionName = Session->SessionName](bool bWasSuccessful) {
+				// We are part of the lobby, so we'll be able to copy the member data
+				CopyLobbyData(LobbyDetails, LobbyDetailsInfo, *Session, true, [this, SessionName = Session->SessionName](bool bWasSuccessful) {
 					if (bWasSuccessful)
 					{
 						if (FNamedOnlineSession* Session = GetNamedSession(SessionName))
@@ -4665,8 +4666,8 @@ void FOnlineSessionEOS::AddLobbySearchResult(const TSharedRef<FLobbyDetailsEOS>&
 
 		SearchResult.Session.SessionInfo = MakeShared<FOnlineSessionInfoEOS>(FOnlineSessionInfoEOS::Create(HostAddr, FUniqueNetIdEOSLobby::Create(UTF8_TO_TCHAR(LobbyDetailsInfo->LobbyId)), LobbyDetails));
 
-		// We copy the lobby data and settings
-		CopyLobbyData(LobbyDetails, LobbyDetailsInfo, SearchResult.Session, Callback);
+		// We copy the lobby data and settings, but not the member data (for search results)
+		CopyLobbyData(LobbyDetails, LobbyDetailsInfo, SearchResult.Session, false, Callback);
 
 		EOS_LobbyDetails_Info_Release(LobbyDetailsInfo);
 
@@ -4684,7 +4685,7 @@ void FOnlineSessionEOS::AddLobbySearchResult(const TSharedRef<FLobbyDetailsEOS>&
 	}
 }
 
-void FOnlineSessionEOS::CopyLobbyData(const TSharedRef<FLobbyDetailsEOS>& LobbyDetails, EOS_LobbyDetails_Info* LobbyDetailsInfo, FOnlineSession& OutSession, const FOnCopyLobbyDataCompleteCallback& Callback)
+void FOnlineSessionEOS::CopyLobbyData(const TSharedRef<FLobbyDetailsEOS>& LobbyDetails, EOS_LobbyDetails_Info* LobbyDetailsInfo, FOnlineSession& OutSession, bool bCopyMemberData, const FOnCopyLobbyDataCompleteCallback& Callback)
 {
 	// This method launches an asynchronous operation, so we'll pass the details handle as a shared ref to make sure it stays alive
 
@@ -4719,66 +4720,70 @@ void FOnlineSessionEOS::CopyLobbyData(const TSharedRef<FLobbyDetailsEOS>& LobbyD
 	// We copy the settings related to lobby attributes
 	CopyLobbyAttributes(*LobbyDetails, OutSession);
 
-	// Then we copy the settings for all lobby members
-	EOS_LobbyDetails_GetMemberCountOptions CountOptions = { };
-	CountOptions.ApiVersion = 1;
-	UE_EOS_CHECK_API_MISMATCH(EOS_LOBBYDETAILS_GETMEMBERCOUNT_API_LATEST, 1);
-	int32 Count = EOS_LobbyDetails_GetMemberCount(LobbyDetails->LobbyDetailsHandle, &CountOptions);
-
-	TArray<EOS_ProductUserId> TargetUserIds;
-	TargetUserIds.Reserve(Count);
-	for (int32 Index = 0; Index < Count; Index++)
+	if (bCopyMemberData)
 	{
-		EOS_LobbyDetails_GetMemberByIndexOptions GetMemberByIndexOptions = { };
-		GetMemberByIndexOptions.ApiVersion = 1;
-		UE_EOS_CHECK_API_MISMATCH(EOS_LOBBYDETAILS_GETMEMBERBYINDEX_API_LATEST, 1);
-		GetMemberByIndexOptions.MemberIndex = Index;
-
-		EOS_ProductUserId TargetUserId = EOS_LobbyDetails_GetMemberByIndex(LobbyDetails->LobbyDetailsHandle, &GetMemberByIndexOptions);
+		// Then we copy the settings for all lobby members
+		EOS_LobbyDetails_GetMemberCountOptions CountOptions = { };
+		CountOptions.ApiVersion = 1;
+		UE_EOS_CHECK_API_MISMATCH(EOS_LOBBYDETAILS_GETMEMBERCOUNT_API_LATEST, 1);
+		int32 Count = EOS_LobbyDetails_GetMemberCount(LobbyDetails->LobbyDetailsHandle, &CountOptions);
 
-		TargetUserIds.Add(TargetUserId);
-	}
+		TArray<EOS_ProductUserId> TargetUserIds;
+		TargetUserIds.Reserve(Count);
+		for (int32 Index = 0; Index < Count; Index++)
+		{
+			EOS_LobbyDetails_GetMemberByIndexOptions GetMemberByIndexOptions = { };
+			GetMemberByIndexOptions.ApiVersion = 1;
+			UE_EOS_CHECK_API_MISMATCH(EOS_LOBBYDETAILS_GETMEMBERBYINDEX_API_LATEST, 1);
+			GetMemberByIndexOptions.MemberIndex = Index;
 
-	if (!TargetUserIds.IsEmpty())
-	{
-		EOSSubsystem->UserManager->ResolveUniqueNetIds(EOSSubsystem->UserManager->GetDefaultLocalUser(), TargetUserIds, [this, LobbyDetails, LobbyId = FUniqueNetIdEOSLobby::Create(UTF8_TO_TCHAR(LobbyDetailsInfo->LobbyId)), OriginalCallback = Callback](TMap<EOS_ProductUserId, FUniqueNetIdEOSRef> ResolvedUniqueNetIds)
+			if (EOS_ProductUserId TargetUserId = EOS_LobbyDetails_GetMemberByIndex(LobbyDetails->LobbyDetailsHandle, &GetMemberByIndexOptions))
 			{
-				FOnlineSession* Session = GetOnlineSessionFromLobbyId(*LobbyId);
-				if (Session)
+				TargetUserIds.Add(TargetUserId);
+			}
+		}
+
+		if (!TargetUserIds.IsEmpty())
+		{
+			EOSSubsystem->UserManager->ResolveUniqueNetIds(EOSSubsystem->UserManager->GetDefaultLocalUser(), TargetUserIds, [this, LobbyDetails, LobbyId = FUniqueNetIdEOSLobby::Create(UTF8_TO_TCHAR(LobbyDetailsInfo->LobbyId)), OriginalCallback = Callback](TMap<EOS_ProductUserId, FUniqueNetIdEOSRef> ResolvedUniqueNetIds)
 				{
-					// One of the resolved ids will be the Owner's, so we'll set that too
-					EOS_LobbyDetails_GetLobbyOwnerOptions GetLobbyOwnerOptions = {};
-					GetLobbyOwnerOptions.ApiVersion = 1;
-					UE_EOS_CHECK_API_MISMATCH(EOS_LOBBYDETAILS_GETLOBBYOWNER_API_LATEST, 1);
+					FOnlineSession* Session = GetOnlineSessionFromLobbyId(*LobbyId);
+					if (Session)
+					{
+						// One of the resolved ids will be the Owner's, so we'll set that too
+						EOS_LobbyDetails_GetLobbyOwnerOptions GetLobbyOwnerOptions = {};
+						GetLobbyOwnerOptions.ApiVersion = 1;
+						UE_EOS_CHECK_API_MISMATCH(EOS_LOBBYDETAILS_GETLOBBYOWNER_API_LATEST, 1);
 
-					const EOS_ProductUserId LobbyOwner = EOS_LobbyDetails_GetLobbyOwner(LobbyDetails->LobbyDetailsHandle, &GetLobbyOwnerOptions);
+						const EOS_ProductUserId LobbyOwner = EOS_LobbyDetails_GetLobbyOwner(LobbyDetails->LobbyDetailsHandle, &GetLobbyOwnerOptions);
 
-					if (FUniqueNetIdEOSRef* OwnerNetId = ResolvedUniqueNetIds.Find(LobbyOwner))
-					{
-						Session->OwningUserId = *OwnerNetId;
-						Session->OwningUserName = EOSSubsystem->UserManager->GetPlayerNickname(**OwnerNetId);
-					}
+						if (FUniqueNetIdEOSRef* OwnerNetId = ResolvedUniqueNetIds.Find(LobbyOwner))
+						{
+							Session->OwningUserId = *OwnerNetId;
+							Session->OwningUserName = EOSSubsystem->UserManager->GetPlayerNickname(**OwnerNetId);
+						}
 
-					for (TMap<EOS_ProductUserId, FUniqueNetIdEOSRef>::TConstIterator It(ResolvedUniqueNetIds); It; ++It)
-					{
-						FSessionSettings& MemberSettings = Session->SessionSettings.MemberSettings.FindOrAdd(It.Value());
+						for (TMap<EOS_ProductUserId, FUniqueNetIdEOSRef>::TConstIterator It(ResolvedUniqueNetIds); It; ++It)
+						{
+							FSessionSettings& MemberSettings = Session->SessionSettings.MemberSettings.FindOrAdd(It.Value());
 
-						CopyLobbyMemberAttributes(*LobbyDetails, It.Key(), MemberSettings);
+							CopyLobbyMemberAttributes(*LobbyDetails, It.Key(), MemberSettings);
+						}
 					}
-				}
 
-				const bool bWasSuccessful = Session != nullptr;
-				OriginalCallback(bWasSuccessful);
-			});
+					const bool bWasSuccessful = Session != nullptr;
+					OriginalCallback(bWasSuccessful);
+				});
+
+			return;
+		}
 	}
-	else
+
+	// ResolveUniqueNetIds is an asynchronous operation, so in the cases where it's not called, we'll delay the execution of this callback to match the flow
+	EOSSubsystem->ExecuteNextTick([Callback]()
 	{
-		// ResolveUniqueNetIds is an asynchronous operation, so we'll delay the execution of this callback to match the flow
-		EOSSubsystem->ExecuteNextTick([Callback]()
-		{
-			Callback(true);
-		});
-	}
+		Callback(true);
+	});
 }
 
 void FOnlineSessionEOS::CopyLobbyAttributes(const FLobbyDetailsEOS& LobbyDetails, FOnlineSession& OutSession)
diff --git a/Engine/Plugins/Online/OnlineSubsystemEOS/Source/OnlineSubsystemEOS/Private/OnlineSessionEOS.h b/Engine/Plugins/Online/OnlineSubsystemEOS/Source/OnlineSubsystemEOS/Private/OnlineSessionEOS.h
--- a/Engine/Plugins/Online/OnlineSubsystemEOS/Source/OnlineSubsystemEOS/Private/OnlineSessionEOS.h
+++ b/Engine/Plugins/Online/OnlineSubsystemEOS/Source/OnlineSubsystemEOS/Private/OnlineSessionEOS.h
@@ -250,7 +250,7 @@ private:
 
 	// Methods to update an OSS Lobby from an API Lobby
 	typedef TFunction<void(bool bWasSuccessful)> FOnCopyLobbyDataCompleteCallback;
-	void CopyLobbyData(const TSharedRef<FLobbyDetailsEOS>& LobbyDetails, EOS_LobbyDetails_Info* LobbyDetailsInfo, FOnlineSession& OutSession, const FOnCopyLobbyDataCompleteCallback& Callback);
+	void CopyLobbyData(const TSharedRef<FLobbyDetailsEOS>& LobbyDetails, EOS_LobbyDetails_Info* LobbyDetailsInfo, FOnlineSession& OutSession, bool bCopyMemberData, const FOnCopyLobbyDataCompleteCallback& Callback);
 	void CopyLobbyAttributes(const FLobbyDetailsEOS& LobbyDetails, FOnlineSession& OutSession);
 	void CopyLobbyMemberAttributes(const FLobbyDetailsEOS& LobbyDetails, const EOS_ProductUserId& TargetUserId, FSessionSettings& OutSessionSettings);
 
