Removed temporary presence attribute logic for Lobbies in OnlineSessionEOS, added support for new API functionality.
Sessions don't have any new API functionality, and LobbiesEOSGS works in a different way, so no changes needed there.

diff --git a/Engine/Plugins/Online/OnlineSubsystemEOS/Source/OnlineSubsystemEOS/Private/OnlineSessionEOS.cpp b/Engine/Plugins/Online/OnlineSubsystemEOS/Source/OnlineSubsystemEOS/Private/OnlineSessionEOS.cpp
--- a/Engine/Plugins/Online/OnlineSubsystemEOS/Source/OnlineSubsystemEOS/Private/OnlineSessionEOS.cpp
+++ b/Engine/Plugins/Online/OnlineSubsystemEOS/Source/OnlineSubsystemEOS/Private/OnlineSessionEOS.cpp
@@ -3927,9 +3927,6 @@ uint32 FOnlineSessionEOS::CreateLobbySession(int32 HostingPlayerNum, FNamedOnlin
 	bool bUseHostMigration = true;
 	Session->SessionSettings.Get(SETTING_HOST_MIGRATION, bUseHostMigration);
 
-	// We set the custom parameter to transmit presence information
-	Session->SessionSettings.Settings.Add(USES_PRESENCE_ATTRIBUTE_KEY, FOnlineSessionSetting(Session->SessionSettings.bUsesPresence, EOnlineDataAdvertisementType::ViaOnlineService));
-
 	if (!Session->SessionSettings.bUsesPresence && (Session->SessionSettings.bAllowJoinViaPresence || Session->SessionSettings.bAllowJoinViaPresenceFriendsOnly))
 	{
 		UE_LOG_ONLINE_SESSION(Warning, TEXT("FOnlineSessionSettings::bUsesPresence is set to false, bAllowJoinViaPresence and bAllowJoinViaPresenceFriendsOnly will be automatically set to false as well"));
@@ -4691,6 +4688,7 @@ void FOnlineSessionEOS::CopyLobbyData(const TSharedRef<FLobbyDetailsEOS>& LobbyD
 
 	OutSession.SessionSettings.bUseLobbiesIfAvailable = true;
 	OutSession.SessionSettings.bIsLANMatch = false;
+	OutSession.SessionSettings.bUsesPresence = LobbyDetailsInfo->bPresenceEnabled == EOS_TRUE;
 	OutSession.SessionSettings.Set(SETTING_HOST_MIGRATION, LobbyDetailsInfo->bAllowHostMigration, EOnlineDataAdvertisementType::DontAdvertise);
 #if WITH_EOS_RTC
 	OutSession.SessionSettings.bUseLobbiesVoiceChatIfAvailable = LobbyDetailsInfo->bRTCRoomEnabled == EOS_TRUE;
@@ -4807,11 +4805,7 @@ void FOnlineSessionEOS::CopyLobbyAttributes(const FLobbyDetailsEOS& LobbyDetails
 		if (ResultCode == EOS_EResult::EOS_Success)
 		{
 			FString Key = UTF8_TO_TCHAR(Attribute->Data->Key);
-			if (Key == USES_PRESENCE_ATTRIBUTE_KEY.ToString())
-			{
-				OutSession.SessionSettings.bUsesPresence = Attribute->Data->Value.AsBool == EOS_TRUE;
-			}
-			else if (Key == TEXT("NumPublicConnections"))
+			if (Key == TEXT("NumPublicConnections"))
 			{
 				OutSession.SessionSettings.NumPublicConnections = Attribute->Data->Value.AsInt64;
 			}
