Integration for EOS_Platform_RTCOptions::BackgroundMode

diff --git a/Engine/Plugins/Online/EOSShared/Source/EOSShared/Private/EOSSDKManager.cpp b/Engine/Plugins/Online/EOSShared/Source/EOSShared/Private/EOSSDKManager.cpp
--- a/Engine/Plugins/Online/EOSShared/Source/EOSShared/Private/EOSSDKManager.cpp
+++ b/Engine/Plugins/Online/EOSShared/Source/EOSShared/Private/EOSSDKManager.cpp
@@ -493,7 +493,8 @@ IEOSPlatformHandlePtr FEOSSDKManager::CreatePlatform(const FString& PlatformConf
 	PlatformRTCOptions.ApiVersion = 2;
 	UE_EOS_CHECK_API_MISMATCH(EOS_PLATFORM_RTCOPTIONS_API_LATEST, 2);
 	PlatformRTCOptions.PlatformSpecificOptions = nullptr;
-	PlatformRTCOptions.BackgroundMode = EOS_ERTCBackgroundMode::EOS_RTCBM_LeaveRooms;
+	PlatformRTCOptions.BackgroundMode = PlatformConfig->RTCBackgroundMode;
+
 	PlatformOptions.RTCOptions = PlatformConfig->bEnableRTC ? &PlatformRTCOptions : nullptr;
 
 	PlatformOptions.IntegratedPlatformOptionsContainerHandle = nullptr;
diff --git a/Engine/Plugins/Online/EOSShared/Source/EOSShared/Private/EOSShared.cpp b/Engine/Plugins/Online/EOSShared/Source/EOSShared/Private/EOSShared.cpp
--- a/Engine/Plugins/Online/EOSShared/Source/EOSShared/Private/EOSShared.cpp
+++ b/Engine/Plugins/Online/EOSShared/Source/EOSShared/Private/EOSShared.cpp
@@ -10,7 +10,6 @@
 #include "eos_auth_types.h"
 #include "eos_friends_types.h"
 #include "eos_presence_types.h"
-#include "eos_types.h"
 
 DEFINE_LOG_CATEGORY(LogEOSSDK);
 
@@ -290,4 +289,22 @@ bool LexFromString(EOS_ELoginCredentialType& OutEnum, const TCHAR* InString)
 	return true;
 }
 
+bool LexFromString(EOS_ERTCBackgroundMode& OutEnum, const TCHAR* InString)
+{
+	if (FCString::Stricmp(InString, TEXT("LeaveRooms")) == 0)
+	{
+		OutEnum = EOS_ERTCBackgroundMode::EOS_RTCBM_LeaveRooms;
+	}
+	else if (FCString::Stricmp(InString, TEXT("KeepRoomsAlive")) == 0)
+	{
+		OutEnum = EOS_ERTCBackgroundMode::EOS_RTCBM_KeepRoomsAlive;
+	}
+	else
+	{
+		checkNoEntry();
+		return false;
+	}
+
+	return true;
+}
 #endif // WITH_EOS_SDK
\ No newline at end of file
diff --git a/Engine/Plugins/Online/EOSShared/Source/EOSShared/Public/EOSShared.h b/Engine/Plugins/Online/EOSShared/Source/EOSShared/Public/EOSShared.h
--- a/Engine/Plugins/Online/EOSShared/Source/EOSShared/Public/EOSShared.h
+++ b/Engine/Plugins/Online/EOSShared/Source/EOSShared/Public/EOSShared.h
@@ -38,6 +38,7 @@
 #include "eos_base.h"
 #include "eos_common.h"
 #include "eos_version.h"
+#include "eos_types.h"
 
 #if defined(DISABLE_EOSVOICECHAT_ENGINE)
 #define WITH_EOS_RTC 0
@@ -83,5 +84,6 @@ EOSSHARED_API const TCHAR* LexToString(const EOS_Presence_EStatus PresenceStatus
 EOSSHARED_API bool LexFromString(EOS_EAuthScopeFlags& OutEnum, const FStringView InString);
 EOSSHARED_API bool LexFromString(EOS_EExternalCredentialType& OutEnum, const TCHAR* InString);
 EOSSHARED_API bool LexFromString(EOS_ELoginCredentialType& OutEnum, const TCHAR* InString);
+EOSSHARED_API bool LexFromString(EOS_ERTCBackgroundMode& OutEnum, const TCHAR* InString);
 
 #endif // WITH_EOS_SDK
\ No newline at end of file
diff --git a/Engine/Plugins/Online/EOSShared/Source/EOSShared/Public/IEOSSDKManager.h b/Engine/Plugins/Online/EOSShared/Source/EOSShared/Public/IEOSSDKManager.h
--- a/Engine/Plugins/Online/EOSShared/Source/EOSShared/Public/IEOSSDKManager.h
+++ b/Engine/Plugins/Online/EOSShared/Source/EOSShared/Public/IEOSSDKManager.h
@@ -25,6 +25,7 @@ struct FEOSSDKPlatformConfig
 	FString OverrideLocaleCode;
 	FString DeploymentId;
 	FString CacheDirectory;
+	EOS_ERTCBackgroundMode RTCBackgroundMode = EOS_ERTCBackgroundMode::EOS_RTCBM_KeepRoomsAlive;
 	bool bIsServer = false;
 	bool bLoadingInEditor = false;
 	bool bDisableOverlay = false;
diff --git a/Engine/Plugins/Online/OnlineSubsystemEOS/Source/OnlineSubsystemEOS/OnlineSubsystemEOS.Build.cs b/Engine/Plugins/Online/OnlineSubsystemEOS/Source/OnlineSubsystemEOS/OnlineSubsystemEOS.Build.cs
--- a/Engine/Plugins/Online/OnlineSubsystemEOS/Source/OnlineSubsystemEOS/OnlineSubsystemEOS.Build.cs
+++ b/Engine/Plugins/Online/OnlineSubsystemEOS/Source/OnlineSubsystemEOS/OnlineSubsystemEOS.Build.cs
@@ -12,6 +12,7 @@ public class OnlineSubsystemEOS : ModuleRules
 		PublicDependencyModuleNames.AddRange(
 			new string[] {
 				"EOSSDK",
+				"EOSShared"
 			}
 		);
 
@@ -21,7 +22,6 @@ public class OnlineSubsystemEOS : ModuleRules
 				"CoreOnline",
 				"CoreUObject",
 				"Engine",
-				"EOSShared",
 				"EOSVoiceChat",
 				"Json",
 				"OnlineBase",
diff --git a/Engine/Plugins/Online/OnlineSubsystemEOS/Source/OnlineSubsystemEOS/Private/EOSSettings.cpp b/Engine/Plugins/Online/OnlineSubsystemEOS/Source/OnlineSubsystemEOS/Private/EOSSettings.cpp
--- a/Engine/Plugins/Online/OnlineSubsystemEOS/Source/OnlineSubsystemEOS/Private/EOSSettings.cpp
+++ b/Engine/Plugins/Online/OnlineSubsystemEOS/Source/OnlineSubsystemEOS/Private/EOSSettings.cpp
@@ -129,6 +129,23 @@ FEOSArtifactSettings ParseArtifactSettingsFromConfigString(const FString& RawLin
 	return Result;
 }
 
+FEOSSettings::FEOSSettings()
+	: RTCBackgroundMode(EOS_ERTCBackgroundMode::EOS_RTCBM_KeepRoomsAlive)
+	, TickBudgetInMilliseconds(0)
+	, TitleStorageReadChunkLength(0)
+	, bEnableOverlay(false)
+	, bEnableSocialOverlay(false)
+	, bEnableEditorOverlay(false)
+	, bUseEAS(false)
+	, bUseEOSConnect(false)
+	, bUseEOSSessions(false)
+	, bMirrorStatsToEOS(false)
+	, bMirrorAchievementsToEOS(false)
+	, bMirrorPresenceToEAS(false)
+{
+
+}
+
 FEOSSettings UEOSSettings::GetSettings()
 {
 	if (UObjectInitialized())
@@ -154,6 +171,13 @@ const FEOSSettings& UEOSSettings::ManualGetSettings()
 
 		GConfig->GetString(INI_SECTION, TEXT("CacheDir"), CachedSettings->CacheDir, GEngineIni);
 		GConfig->GetString(INI_SECTION, TEXT("DefaultArtifactName"), CachedSettings->DefaultArtifactName, GEngineIni);
+		CachedSettings->RTCBackgroundMode = EOS_ERTCBackgroundMode::EOS_RTCBM_KeepRoomsAlive;
+		FString RTCBackgroundModeStr;
+		GConfig->GetString(INI_SECTION, TEXT("RTCBackgroundMode"), RTCBackgroundModeStr, GEngineIni);
+		if (!RTCBackgroundModeStr.IsEmpty())
+		{
+			LexFromString(CachedSettings->RTCBackgroundMode, *RTCBackgroundModeStr);
+		}
 		GConfig->GetInt(INI_SECTION, TEXT("TickBudgetInMilliseconds"), CachedSettings->TickBudgetInMilliseconds, GEngineIni);
 		GConfig->GetInt(INI_SECTION, TEXT("TitleStorageReadChunkLength"), CachedSettings->TitleStorageReadChunkLength, GEngineIni);
 		GConfig->GetBool(INI_SECTION, TEXT("bEnableOverlay"), CachedSettings->bEnableOverlay, GEngineIni);
@@ -179,6 +203,11 @@ FEOSSettings UEOSSettings::ToNative() const
 
 	Native.CacheDir = CacheDir;
 	Native.DefaultArtifactName = DefaultArtifactName;
+	Native.RTCBackgroundMode = EOS_ERTCBackgroundMode::EOS_RTCBM_KeepRoomsAlive;
+	if (!RTCBackgroundMode.IsEmpty())
+	{
+		LexFromString(Native.RTCBackgroundMode, *RTCBackgroundMode);
+	}
 	Native.TickBudgetInMilliseconds = TickBudgetInMilliseconds;
 	Native.TitleStorageReadChunkLength = TitleStorageReadChunkLength;
 	Native.bEnableOverlay = bEnableOverlay;
diff --git a/Engine/Plugins/Online/OnlineSubsystemEOS/Source/OnlineSubsystemEOS/Private/OnlineSubsystemEOS.cpp b/Engine/Plugins/Online/OnlineSubsystemEOS/Source/OnlineSubsystemEOS/Private/OnlineSubsystemEOS.cpp
--- a/Engine/Plugins/Online/OnlineSubsystemEOS/Source/OnlineSubsystemEOS/Private/OnlineSubsystemEOS.cpp
+++ b/Engine/Plugins/Online/OnlineSubsystemEOS/Source/OnlineSubsystemEOS/Private/OnlineSubsystemEOS.cpp
@@ -238,7 +238,7 @@ bool FOnlineSubsystemEOS::PlatformCreate()
 	RtcOptions.ApiVersion = 2;
 	UE_EOS_CHECK_API_MISMATCH(EOS_PLATFORM_RTCOPTIONS_API_LATEST, 2);
 	RtcOptions.PlatformSpecificOptions = nullptr;
-	RtcOptions.BackgroundMode = EOS_ERTCBackgroundMode::EOS_RTCBM_LeaveRooms;
+	RtcOptions.BackgroundMode = EOSSettings.RTCBackgroundMode;
 	PlatformOptions.RTCOptions = &RtcOptions;
 #endif
 
diff --git a/Engine/Plugins/Online/OnlineSubsystemEOS/Source/OnlineSubsystemEOS/Public/EOSSettings.h b/Engine/Plugins/Online/OnlineSubsystemEOS/Source/OnlineSubsystemEOS/Public/EOSSettings.h
--- a/Engine/Plugins/Online/OnlineSubsystemEOS/Source/OnlineSubsystemEOS/Public/EOSSettings.h
+++ b/Engine/Plugins/Online/OnlineSubsystemEOS/Source/OnlineSubsystemEOS/Public/EOSSettings.h
@@ -5,6 +5,7 @@
 #include "UObject/ObjectMacros.h"
 #include "Engine/RuntimeOptionsBase.h"
 #include "Engine/DataAsset.h"
+#include "EOSShared.h"
 #include "EOSSettings.generated.h"
 
 /** Native version of the UObject based config data */
@@ -66,8 +67,11 @@ public:
 /** Native version of the UObject based config data */
 struct FEOSSettings
 {
+	FEOSSettings();
+
 	FString CacheDir;
 	FString DefaultArtifactName;
+	EOS_ERTCBackgroundMode RTCBackgroundMode;
 	int32 TickBudgetInMilliseconds;
 	int32 TitleStorageReadChunkLength;
 	bool bEnableOverlay;
@@ -103,6 +107,10 @@ public:
 	UPROPERTY(Config, EditAnywhere, BlueprintReadOnly, Category="EOS Settings")
 	FString DefaultArtifactName;
 
+	/** The preferred background mode to be used by RTC services */
+	UPROPERTY(Config, EditAnywhere, BlueprintReadOnly, Category="EOS Settings")
+	FString RTCBackgroundMode;
+
 	/** Used to throttle how much time EOS ticking can take */
 	UPROPERTY(Config, EditAnywhere, BlueprintReadOnly, Category="EOS Settings")
 	int32 TickBudgetInMilliseconds = 0;
diff --git a/Engine/Plugins/Online/VoiceChat/EOSVoiceChat/Source/EOSVoiceChat/Private/EOSVoiceChat.cpp b/Engine/Plugins/Online/VoiceChat/EOSVoiceChat/Source/EOSVoiceChat/Private/EOSVoiceChat.cpp
--- a/Engine/Plugins/Online/VoiceChat/EOSVoiceChat/Source/EOSVoiceChat/Private/EOSVoiceChat.cpp
+++ b/Engine/Plugins/Online/VoiceChat/EOSVoiceChat/Source/EOSVoiceChat/Private/EOSVoiceChat.cpp
@@ -156,7 +156,12 @@ void FEOSVoiceChat::Initialize(const FOnVoiceChatInitializeCompleteDelegate& Ini
 					EOS_Platform_RTCOptions PlatformRTCOptions = {};
 					PlatformRTCOptions.ApiVersion = 2;
 					UE_EOS_CHECK_API_MISMATCH(EOS_PLATFORM_RTCOPTIONS_API_LATEST, 2);
-					PlatformRTCOptions.BackgroundMode = EOS_ERTCBackgroundMode::EOS_RTCBM_LeaveRooms;
+					PlatformRTCOptions.BackgroundMode = EOS_ERTCBackgroundMode::EOS_RTCBM_KeepRoomsAlive;
+					FString RTCBackgroundModeStr;
+					if(GConfig->GetString(TEXT("EOSVoiceChat"), TEXT("RTCBackgroundMode"), RTCBackgroundModeStr, GEngineIni))
+					{
+						LexFromString(PlatformRTCOptions.BackgroundMode, *RTCBackgroundModeStr);					
+					}
 					PlatformOptions.RTCOptions = &PlatformRTCOptions;
 
 					EosPlatformHandle = EOSPlatformCreate(PlatformOptions);
