- Fix achievement strings not being converted from UTF8

diff --git a/Engine/Plugins/Online/OnlineSubsystemEOS/Source/OnlineSubsystemEOS/Private/OnlineAchievementsEOS.cpp b/Engine/Plugins/Online/OnlineSubsystemEOS/Source/OnlineSubsystemEOS/Private/OnlineAchievementsEOS.cpp
--- a/Engine/Plugins/Online/OnlineSubsystemEOS/Source/OnlineSubsystemEOS/Private/OnlineAchievementsEOS.cpp
+++ b/Engine/Plugins/Online/OnlineSubsystemEOS/Source/OnlineSubsystemEOS/Private/OnlineAchievementsEOS.cpp
@@ -223,9 +223,9 @@ void FOnlineAchievementsEOS::QueryAchievementDescriptions(const FUniqueNetId& Pl
 					// Work around for the ID not being part of the description
 					CachedAchievementDefinitionsMap.Add(Definition->AchievementId, Desc);
 
-					Desc->Title = FText::FromString(Definition->DisplayName);
-					Desc->LockedDesc = FText::FromString(Definition->LockedDescription);
-					Desc->UnlockedDesc = FText::FromString(Definition->CompletionDescription);
+					Desc->Title = FText::FromString(UTF8_TO_TCHAR(Definition->DisplayName));
+					Desc->LockedDesc = FText::FromString(UTF8_TO_TCHAR(Definition->LockedDescription));
+					Desc->UnlockedDesc = FText::FromString(UTF8_TO_TCHAR(Definition->CompletionDescription));
 					Desc->bIsHidden = Definition->bIsHidden == EOS_TRUE;
 
 					EOS_Achievements_Definition_Release(Definition);
