Refactor and fix for proper functionality in EOS_UserInfo_CopyBestDisplayName call sites

diff --git a/Engine/Plugins/Online/EOSShared/Source/EOSShared/Private/EOSShared.cpp b/Engine/Plugins/Online/EOSShared/Source/EOSShared/Private/EOSShared.cpp
--- a/Engine/Plugins/Online/EOSShared/Source/EOSShared/Private/EOSShared.cpp
+++ b/Engine/Plugins/Online/EOSShared/Source/EOSShared/Private/EOSShared.cpp
@@ -289,6 +289,11 @@ bool LexFromString(EOS_ELoginCredentialType& OutEnum, const TCHAR* InString)
 	return true;
 }
 
+FString GetBestDisplayNameStr(const EOS_UserInfo_BestDisplayName& BestDisplayName)
+{
+	return FString(UTF8_TO_TCHAR(BestDisplayName.Nickname ? BestDisplayName.Nickname : BestDisplayName.DisplayNameSanitized ? BestDisplayName.DisplayNameSanitized : BestDisplayName.DisplayName));
+}
+
 bool LexFromString(EOS_ERTCBackgroundMode& OutEnum, const TCHAR* InString)
 {
 	if (FCString::Stricmp(InString, TEXT("LeaveRooms")) == 0)
diff --git a/Engine/Plugins/Online/EOSShared/Source/EOSShared/Public/EOSShared.h b/Engine/Plugins/Online/EOSShared/Source/EOSShared/Public/EOSShared.h
--- a/Engine/Plugins/Online/EOSShared/Source/EOSShared/Public/EOSShared.h
+++ b/Engine/Plugins/Online/EOSShared/Source/EOSShared/Public/EOSShared.h
@@ -37,8 +37,9 @@
 #endif
 #include "eos_base.h"
 #include "eos_common.h"
-#include "eos_version.h"
 #include "eos_types.h"
+#include "eos_userinfo_types.h"
+#include "eos_version.h"
 
 #if defined(DISABLE_EOSVOICECHAT_ENGINE)
 #define WITH_EOS_RTC 0
@@ -86,4 +87,7 @@ EOSSHARED_API bool LexFromString(EOS_EExternalCredentialType& OutEnum, const TCH
 EOSSHARED_API bool LexFromString(EOS_ELoginCredentialType& OutEnum, const TCHAR* InString);
 EOSSHARED_API bool LexFromString(EOS_ERTCBackgroundMode& OutEnum, const TCHAR* InString);
 
+/** Extracts the display name FString from a EOS_UserInfo_BestDisplayName using the following logic: Nickname > DisplayNameSanitized > DisplayName */
+EOSSHARED_API FString GetBestDisplayNameStr(const EOS_UserInfo_BestDisplayName& BestDisplayName);
+
 #endif // WITH_EOS_SDK
\ No newline at end of file
diff --git a/Engine/Plugins/Online/OnlineServicesEOS/Source/Private/Online/AuthEOS.cpp b/Engine/Plugins/Online/OnlineServicesEOS/Source/Private/Online/AuthEOS.cpp
--- a/Engine/Plugins/Online/OnlineServicesEOS/Source/Private/Online/AuthEOS.cpp
+++ b/Engine/Plugins/Online/OnlineServicesEOS/Source/Private/Online/AuthEOS.cpp
@@ -219,7 +219,7 @@ TOnlineAsyncOpHandle<FAuthLogin> FAuthEOS::Login(FAuthLogin::Params&& Params)
 
 		if (CopyBestDisplayNameResult == EOS_EResult::EOS_Success)
 		{
-			AccountInfoEOS->Attributes.Emplace(AccountAttributeData::DisplayName, UTF8_TO_TCHAR(BestDisplayName->DisplayNameSanitized));
+			AccountInfoEOS->Attributes.Emplace(AccountAttributeData::DisplayName, *GetBestDisplayNameStr(*BestDisplayName));
 			EOS_UserInfo_BestDisplayName_Release(BestDisplayName);
 		}
 		else
diff --git a/Engine/Plugins/Online/OnlineServicesEOS/Source/Private/Online/UserInfoEOS.cpp b/Engine/Plugins/Online/OnlineServicesEOS/Source/Private/Online/UserInfoEOS.cpp
--- a/Engine/Plugins/Online/OnlineServicesEOS/Source/Private/Online/UserInfoEOS.cpp
+++ b/Engine/Plugins/Online/OnlineServicesEOS/Source/Private/Online/UserInfoEOS.cpp
@@ -135,7 +135,7 @@ TOnlineResult<FGetUserInfo> FUserInfoEOS::GetUserInfo(FGetUserInfo::Params&& Par
 
 	TSharedRef<FUserInfo> UserInfo = MakeShared<FUserInfo>();
 	UserInfo->AccountId = Params.AccountId;
-	UserInfo->DisplayName = UTF8_TO_TCHAR(EosBestDisplayName->DisplayNameSanitized);
+	UserInfo->DisplayName = GetBestDisplayNameStr(*EosBestDisplayName);
 
 	return TOnlineResult<FGetUserInfo>({UserInfo});
 }
diff --git a/Engine/Plugins/Online/OnlineSubsystemEOS/Source/OnlineSubsystemEOS/Private/UserManagerEOS.cpp b/Engine/Plugins/Online/OnlineSubsystemEOS/Source/OnlineSubsystemEOS/Private/UserManagerEOS.cpp
--- a/Engine/Plugins/Online/OnlineSubsystemEOS/Source/OnlineSubsystemEOS/Private/UserManagerEOS.cpp
+++ b/Engine/Plugins/Online/OnlineSubsystemEOS/Source/OnlineSubsystemEOS/Private/UserManagerEOS.cpp
@@ -1081,10 +1081,10 @@ void FUserManagerEOS::UpdateUserInfo(IAttributeAccessInterfaceRef AttributeAcces
 	}
 
 	EOS_UserInfo_CopyBestDisplayNameOptions BestDisplayNameOptions = { };
-	Options.ApiVersion = 1;
+	BestDisplayNameOptions.ApiVersion = 1;
 	UE_EOS_CHECK_API_MISMATCH(EOS_USERINFO_COPYBESTDISPLAYNAME_API_LATEST, 1);
-	Options.LocalUserId = LocalId;
-	Options.TargetUserId = AccountId;
+	BestDisplayNameOptions.LocalUserId = LocalId;
+	BestDisplayNameOptions.TargetUserId = AccountId;
 
 	EOS_UserInfo_BestDisplayName* BestDisplayName;
 	EOS_EResult BestDisplayNameResult = EOS_UserInfo_CopyBestDisplayName(EOSSubsystem->UserInfoHandle, &BestDisplayNameOptions, &BestDisplayName);
@@ -1103,7 +1103,8 @@ void FUserManagerEOS::UpdateUserInfo(IAttributeAccessInterfaceRef AttributeAcces
 
 	if (BestDisplayNameResult == EOS_EResult::EOS_Success)
 	{
-		AttributeAccessRef->SetInternalAttribute(USER_ATTR_DISPLAY_NAME, UTF8_TO_TCHAR(BestDisplayName->DisplayNameSanitized));
+		// We'll prioritize which name is chosen: Nickname > DisplayNameSanitized > DisplayName
+		AttributeAccessRef->SetInternalAttribute(USER_ATTR_DISPLAY_NAME, *GetBestDisplayNameStr(*BestDisplayName));
 
 		EOS_UserInfo_BestDisplayName_Release(BestDisplayName);
 	}
