- Make public IUniqueNetIdEOSPlus interface so people can cast to it and get the base / eos net id within

diff --git a/Engine/Plugins/Online/OnlineSubsystemEOS/Source/OnlineSubsystemEOSPlus/Private/OnlineUserEOSPlus.h b/Engine/Plugins/Online/OnlineSubsystemEOS/Source/OnlineSubsystemEOSPlus/Private/OnlineUserEOSPlus.h
--- a/Engine/Plugins/Online/OnlineSubsystemEOS/Source/OnlineSubsystemEOSPlus/Private/OnlineUserEOSPlus.h
+++ b/Engine/Plugins/Online/OnlineSubsystemEOS/Source/OnlineSubsystemEOSPlus/Private/OnlineUserEOSPlus.h
@@ -5,6 +5,7 @@
 #include "CoreMinimal.h"
 #include "Online/CoreOnline.h"
 #include "OnlineSubsystem.h"
+#include "OnlineSubsystemEOSPlusTypes.h"
 #include "Interfaces/OnlineUserInterface.h"
 #include "Interfaces/OnlineIdentityInterface.h"
 #include "Interfaces/OnlineFriendsInterface.h"
@@ -26,7 +27,7 @@ using FUniqueNetIdEOSPlusRef = TSharedRef<const class FUniqueNetIdEOSPlus>;
  * Unique net id wrapper for a EOS plus another account id. The underlying string is a combination
  * of both account ids concatenated
  */
-class FUniqueNetIdEOSPlus : public FUniqueNetId
+class FUniqueNetIdEOSPlus : public IUniqueNetIdEOSPlus
 {
 public:
 	template<typename... TArgs>
@@ -44,15 +45,10 @@ public:
 	virtual FString ToDebugString() const override;
 // ~FUniqueNetId interface
 
-	FUniqueNetIdPtr GetBaseNetId() const
-	{
-		return BaseUniqueNetId;
-	}
-
-	FUniqueNetIdPtr GetEOSNetId() const
-	{
-		return EOSUniqueNetId;
-	}
+// IUniqueNetIdEOSPlus interface
+	virtual FUniqueNetIdPtr GetBaseNetId() const override { return BaseUniqueNetId; }
+	virtual FUniqueNetIdPtr GetEOSNetId() const override { return EOSUniqueNetId; }
+// ~IUniqueNetIdEOSPlus interface
 
 	/** global static instance of invalid (zero) id */
 	static const FUniqueNetIdEOSPlusRef& EmptyId()
diff --git a/Engine/Plugins/Online/OnlineSubsystemEOS/Source/OnlineSubsystemEOSPlus/Public/OnlineSubsystemEOSPlusTypes.h b/Engine/Plugins/Online/OnlineSubsystemEOS/Source/OnlineSubsystemEOSPlus/Public/OnlineSubsystemEOSPlusTypes.h
new file mode 100644
--- /dev/null
+++ b/Engine/Plugins/Online/OnlineSubsystemEOS/Source/OnlineSubsystemEOSPlus/Public/OnlineSubsystemEOSPlusTypes.h
@@ -0,0 +1,21 @@
+// Copyright Epic Games, Inc. All Rights Reserved.
+
+#pragma once
+
+#include "Online/CoreOnline.h"
+
+/** Unique net id for OnlineSubsystemEOS */
+class IUniqueNetIdEOSPlus : public FUniqueNetId
+{
+public:
+	/**
+	 * Get the Base aka Platform net id. May be null.
+	 * @return the Base net id, or null.
+	 */
+	virtual FUniqueNetIdPtr GetBaseNetId() const = 0;
+	/**
+	 * Get the EOS net Id. May be null.
+	 * @return the EOS net id, or null.
+	 */
+	virtual FUniqueNetIdPtr GetEOSNetId() const = 0;
+};
\ No newline at end of file
