Fix for PUID strings being kept in the wrong format

diff --git a/Engine/Plugins/Online/OnlineServicesEOS/Source/Private/Online/OnlineIdEOS.cpp b/Engine/Plugins/Online/OnlineServicesEOS/Source/Private/Online/OnlineIdEOS.cpp
--- a/Engine/Plugins/Online/OnlineServicesEOS/Source/Private/Online/OnlineIdEOS.cpp
+++ b/Engine/Plugins/Online/OnlineServicesEOS/Source/Private/Online/OnlineIdEOS.cpp
@@ -112,7 +112,7 @@ FAccountId FOnlineAccountIdRegistryEOS::FromReplicationData(const TArray<uint8>&
 	if (EasHexBufferLength > 0)
 	{
 		char EasBuffer[EOS_EPICACCOUNTID_MAX_LENGTH + 1] = { 0 };
-		UE::String::BytesToHex(TConstArrayView<uint8>(ReplicationData.GetData() + 1, EasHexBufferLength), EasBuffer);
+		UE::String::BytesToHexLower(TConstArrayView<uint8>(ReplicationData.GetData() + 1, EasHexBufferLength), EasBuffer);
 		EpicAccountId = EOS_EpicAccountId_FromString(EasBuffer);
 	}
 
diff --git a/Engine/Plugins/Online/OnlineSubsystemEOS/Source/OnlineSubsystemEOS/Private/OnlineSubsystemEOSTypes.cpp b/Engine/Plugins/Online/OnlineSubsystemEOS/Source/OnlineSubsystemEOS/Private/OnlineSubsystemEOSTypes.cpp
--- a/Engine/Plugins/Online/OnlineSubsystemEOS/Source/OnlineSubsystemEOS/Private/OnlineSubsystemEOSTypes.cpp
+++ b/Engine/Plugins/Online/OnlineSubsystemEOS/Source/OnlineSubsystemEOS/Private/OnlineSubsystemEOSTypes.cpp
@@ -76,14 +76,14 @@ FUniqueNetIdEOS::FUniqueNetIdEOS(const uint8* Bytes, int32 Size)
 	const bool bIsEasNonZero = Algo::AnyOf(TArrayView<const uint8>(Bytes, ID_HALF_BYTE_SIZE));
 	if (bIsEasNonZero)
 	{
-		const FString EpicAccountIdStr = BytesToHex(Bytes, ID_HALF_BYTE_SIZE);
+		const FString EpicAccountIdStr = BytesToHexLower(Bytes, ID_HALF_BYTE_SIZE);
 		EpicAccountId = EOS_EpicAccountId_FromString(TCHAR_TO_UTF8(*EpicAccountIdStr));
 	}
 
 	const bool bIsPuidNonZero = Algo::AnyOf(TArrayView<const uint8>(Bytes + ID_HALF_BYTE_SIZE, ID_HALF_BYTE_SIZE));
 	if (bIsPuidNonZero)
 	{
-		const FString ProductUserIdStr = BytesToHex(Bytes + ID_HALF_BYTE_SIZE, ID_HALF_BYTE_SIZE);
+		const FString ProductUserIdStr = BytesToHexLower(Bytes + ID_HALF_BYTE_SIZE, ID_HALF_BYTE_SIZE);
 		ProductUserId = EOS_ProductUserId_FromString(TCHAR_TO_UTF8(*ProductUserIdStr));
 	}
 }
