Added logic to resolve the Lobby Owner NetId for lobby search results, in which we don't resolve the ids of the other lobby members.

diff --git a/Engine/Plugins/Online/OnlineSubsystemEOS/Source/OnlineSubsystemEOS/Private/OnlineSessionEOS.cpp b/Engine/Plugins/Online/OnlineSubsystemEOS/Source/OnlineSubsystemEOS/Private/OnlineSessionEOS.cpp
--- a/Engine/Plugins/Online/OnlineSubsystemEOS/Source/OnlineSubsystemEOS/Private/OnlineSessionEOS.cpp
+++ b/Engine/Plugins/Online/OnlineSubsystemEOS/Source/OnlineSubsystemEOS/Private/OnlineSessionEOS.cpp
@@ -4776,6 +4776,33 @@ void FOnlineSessionEOS::CopyLobbyData(const TSharedRef<FLobbyDetailsEOS>& LobbyD
 			return;
 		}
 	}
+	else // If we should not copy the member data, we still need to copy the Owner's Id and Name
+	{
+		EOS_LobbyDetails_GetLobbyOwnerOptions GetLobbyOwnerOptions = {};
+		GetLobbyOwnerOptions.ApiVersion = 1;
+		UE_EOS_CHECK_API_MISMATCH(EOS_LOBBYDETAILS_GETLOBBYOWNER_API_LATEST, 1);
+
+		const EOS_ProductUserId LobbyOwner = EOS_LobbyDetails_GetLobbyOwner(LobbyDetails->LobbyDetailsHandle, &GetLobbyOwnerOptions);
+
+		EOSSubsystem->UserManager->ResolveUniqueNetIds(EOSSubsystem->UserManager->GetDefaultLocalUser(), { LobbyOwner }, [this, LobbyOwner, LobbyDetails, LobbyId = FUniqueNetIdEOSLobby::Create(UTF8_TO_TCHAR(LobbyDetailsInfo->LobbyId)), OriginalCallback = Callback](TMap<EOS_ProductUserId, FUniqueNetIdEOSRef> ResolvedUniqueNetIds)
+			{
+				FOnlineSession* Session = GetOnlineSessionFromLobbyId(*LobbyId);
+				if (Session)
+				{
+					FUniqueNetIdEOSRef* OwnerNetId = ResolvedUniqueNetIds.Find(LobbyOwner);
+					if (ensure(OwnerNetId))
+					{
+						Session->OwningUserId = *OwnerNetId;
+						Session->OwningUserName = EOSSubsystem->UserManager->GetPlayerNickname(**OwnerNetId);
+					}
+				}
+
+				const bool bWasSuccessful = Session != nullptr;
+				OriginalCallback(bWasSuccessful);
+			});
+
+		return;
+	}
 
 	// ResolveUniqueNetIds is an asynchronous operation, so in the cases where it's not called, we'll delay the execution of this callback to match the flow
 	EOSSubsystem->ExecuteNextTick([Callback]()
