Added logic to prefer AccountPortal + PersistentAuth login when using OSS EOS on Switch with bUseEAS enabled
Added logic to save, load and delete AccountPortal auth token and token expiration locally using ISaveGameSystem
Refactored login flows to minimize code duplication

diff --git a/Engine/Config/BaseEngine.ini b/Engine/Config/BaseEngine.ini
--- a/Engine/Config/BaseEngine.ini
+++ b/Engine/Config/BaseEngine.ini
@@ -2295,6 +2295,7 @@ bEnabled=false
 +AuthScopeFlags="BasicProfile"
 +AuthScopeFlags="FriendsList"
 +AuthScopeFlags="Presence"
+bPreferPersistentAuth=false
 
 [OnlineSubsystemNull]
 bEnabled=true
diff --git a/Engine/Plugins/Online/OnlineSubsystemEOS/Source/OnlineSubsystemEOS/OnlineSubsystemEOS.Build.cs b/Engine/Plugins/Online/OnlineSubsystemEOS/Source/OnlineSubsystemEOS/OnlineSubsystemEOS.Build.cs
--- a/Engine/Plugins/Online/OnlineSubsystemEOS/Source/OnlineSubsystemEOS/OnlineSubsystemEOS.Build.cs
+++ b/Engine/Plugins/Online/OnlineSubsystemEOS/Source/OnlineSubsystemEOS/OnlineSubsystemEOS.Build.cs
@@ -37,9 +37,11 @@ public class OnlineSubsystemEOS : ModuleRules
 		PrivateDefinitions.Add("USE_XBL_XSTS_TOKEN=" + (bUseXblXstsToken ? "1" : "0"));
 		PrivateDefinitions.Add("USE_PSN_ID_TOKEN=" + (bUsePsnIdToken ? "1" : "0"));
 		PrivateDefinitions.Add("ADD_USER_LOGIN_INFO=" + (bAddUserLoginInfo ? "1" : "0"));
+		PrivateDefinitions.Add("EOS_AUTH_TOKEN_SAVEGAME_STORAGE=" + (bAuthTokenSavegameStorage ? "1" : "0"));
 	}
 
 	protected virtual bool bUseXblXstsToken { get { return false; } }
 	protected virtual bool bUsePsnIdToken { get { return false; } }
 	protected virtual bool bAddUserLoginInfo { get { return false; } }
+	protected virtual bool bAuthTokenSavegameStorage { get { return false; } }
 }
diff --git a/Engine/Plugins/Online/OnlineSubsystemEOS/Source/OnlineSubsystemEOS/Private/EOSSettings.cpp b/Engine/Plugins/Online/OnlineSubsystemEOS/Source/OnlineSubsystemEOS/Private/EOSSettings.cpp
--- a/Engine/Plugins/Online/OnlineSubsystemEOS/Source/OnlineSubsystemEOS/Private/EOSSettings.cpp
+++ b/Engine/Plugins/Online/OnlineSubsystemEOS/Source/OnlineSubsystemEOS/Private/EOSSettings.cpp
@@ -136,6 +136,7 @@ FEOSSettings::FEOSSettings()
 	, bEnableOverlay(false)
 	, bEnableSocialOverlay(false)
 	, bEnableEditorOverlay(false)
+	, bPreferPersistentAuth(false)
 	, bUseEAS(false)
 	, bUseEOSConnect(false)
 	, bUseEOSSessions(false)
@@ -183,6 +184,7 @@ const FEOSSettings& UEOSSettings::ManualGetSettings()
 		GConfig->GetBool(INI_SECTION, TEXT("bEnableOverlay"), CachedSettings->bEnableOverlay, GEngineIni);
 		GConfig->GetBool(INI_SECTION, TEXT("bEnableSocialOverlay"), CachedSettings->bEnableSocialOverlay, GEngineIni);
 		GConfig->GetBool(INI_SECTION, TEXT("bEnableEditorOverlay"), CachedSettings->bEnableEditorOverlay, GEngineIni);
+		GConfig->GetBool(INI_SECTION, TEXT("bPreferPersistentAuth"), CachedSettings->bPreferPersistentAuth, GEngineIni);
 		GConfig->GetBool(INI_SECTION, TEXT("bUseEAS"), CachedSettings->bUseEAS, GEngineIni);
 		GConfig->GetBool(INI_SECTION, TEXT("bUseEOSConnect"), CachedSettings->bUseEOSConnect, GEngineIni);
 		GConfig->GetBool(INI_SECTION, TEXT("bUseEOSSessions"), CachedSettings->bUseEOSSessions, GEngineIni);
@@ -213,6 +215,7 @@ FEOSSettings UEOSSettings::ToNative() const
 	Native.bEnableOverlay = bEnableOverlay;
 	Native.bEnableSocialOverlay = bEnableSocialOverlay;
 	Native.bEnableEditorOverlay = bEnableEditorOverlay;
+	Native.bPreferPersistentAuth = bPreferPersistentAuth;
 	Native.bUseEAS = bUseEAS;
 	Native.bUseEOSConnect = bUseEOSConnect;
 	Native.bUseEOSSessions = bUseEOSSessions;
diff --git a/Engine/Plugins/Online/OnlineSubsystemEOS/Source/OnlineSubsystemEOS/Private/UserManagerEOS.cpp b/Engine/Plugins/Online/OnlineSubsystemEOS/Source/OnlineSubsystemEOS/Private/UserManagerEOS.cpp
--- a/Engine/Plugins/Online/OnlineSubsystemEOS/Source/OnlineSubsystemEOS/Private/UserManagerEOS.cpp
+++ b/Engine/Plugins/Online/OnlineSubsystemEOS/Source/OnlineSubsystemEOS/Private/UserManagerEOS.cpp
@@ -14,7 +14,11 @@
 #include "OnlineSubsystemEOS.h"
 #include "OnlineSubsystemEOSPrivate.h"
 #include "OnlineSubsystemNames.h"
+#include "PlatformFeatures.h"
+#include "SaveGameSystem.h"
 #include "SocketSubsystem.h"
+#include "Serialization/JsonSerializable.h"
+#include "Serialization/JsonSerializerMacros.h"
 
 #include COMPILED_PLATFORM_HEADER(EOSHelpers.h)
 
@@ -395,20 +399,25 @@ bool FUserManagerEOS::Login(int32 LocalUserNum, const FOnlineAccountCredentials&
 		return true;
 	}
 
-	// Can't be const because LoginViaExternalAuth might need to set the type
-	FOnlineAccountCredentials CurrentLoginCredentials(AccountCredentials);
-
 	// See if we are configured to just use EOS and not EAS
 	if (!EOSSubsystem->bIsDefaultOSS && !EOSSubsystem->bIsPlatformOSS && !Settings.bUseEAS && Settings.bUseEOSConnect)
 	{
 		// Call the EOS + Platform login path
-		return ConnectLoginNoEAS(LocalUserNum, CurrentLoginCredentials);
+		return ConnectLoginNoEAS(LocalUserNum, AccountCredentials);
 	}
 
 	// See if we are logging in using platform credentials to link to EAS
 	if (!EOSSubsystem->bIsDefaultOSS && !EOSSubsystem->bIsPlatformOSS && Settings.bUseEAS)
 	{
-		LoginViaExternalAuth(LocalUserNum, CurrentLoginCredentials);
+		if (Settings.bPreferPersistentAuth)
+		{
+			LoginViaPersistentAuth(LocalUserNum, AccountCredentials);
+		}
+		else
+		{
+			LoginViaExternalAuth(LocalUserNum, AccountCredentials);
+		}
+
 		return true;
 	}
 
@@ -417,6 +426,99 @@ bool FUserManagerEOS::Login(int32 LocalUserNum, const FOnlineAccountCredentials&
 	return true;
 }
 
+struct FSavedEOSAuthToken : public FJsonSerializable
+{
+	FString Token;
+	FDateTime ExpirationDate;
+
+	BEGIN_JSON_SERIALIZER
+		JSON_SERIALIZE("token", Token);
+		JSON_SERIALIZE_DATETIME_UNIX_TIMESTAMP_MILLISECONDS("expirationDate", ExpirationDate);
+	END_JSON_SERIALIZER
+};
+
+#define EOS_EPIC_AUTH_TOKEN_FILENAME_SUFFIX TEXT("EpicAuthToken")
+
+FString FUserManagerEOS::GetEOSAuthTokenFilename()
+{
+	return EOSSubsystem->ProductId + TEXT("_") + EOS_EPIC_AUTH_TOKEN_FILENAME_SUFFIX;
+}
+
+void FUserManagerEOS::LoginViaPersistentAuth(int32 LocalUserNum, const FOnlineAccountCredentials& AccountCredentials)
+{
+	FPlatformEOSHelpersPtr EOSHelpers = EOSSubsystem->GetEOSHelpers();
+
+#if EOS_AUTH_TOKEN_SAVEGAME_STORAGE
+	ISaveGameSystem* SaveGameSystem = IPlatformFeaturesModule::Get().GetSaveGameSystem();
+
+	FString SavedEOSAuthTokenFilename = GetEOSAuthTokenFilename();
+
+	if (SaveGameSystem->DoesSaveGameExist(*SavedEOSAuthTokenFilename, LocalUserNum))
+	{
+		TArray<uint8> SavedEOSAuthTokenData;
+
+		if (SaveGameSystem->LoadGame(false, *SavedEOSAuthTokenFilename, LocalUserNum, SavedEOSAuthTokenData))
+		{
+			const FString SavedEOSAuthTokenStr = BytesToString(SavedEOSAuthTokenData.GetData(), SavedEOSAuthTokenData.Num());
+
+			FSavedEOSAuthToken SavedEOSAuthToken;
+			SavedEOSAuthToken.FromJson(SavedEOSAuthTokenStr);
+
+			UE_LOG_ONLINE(Verbose, TEXT("Loaded Saved EOS Auth Token from File [%s]. Contents: [%s]"), *SavedEOSAuthTokenFilename, UE_BUILD_SHIPPING ? TEXT("<redacted>") : *SavedEOSAuthTokenStr);
+
+			// If the token has expired, we'll process login via Account Portal below
+			if (SavedEOSAuthToken.ExpirationDate > FDateTime::Now())
+			{
+				// TODO: We should renew it a configurable amount of time (day or week) before it expires, to avoid having it expire in the middle of a game session
+#endif
+				EOS_Auth_Credentials Credentials = {};
+				Credentials.ApiVersion = 4;
+				UE_EOS_CHECK_API_MISMATCH(EOS_AUTH_CREDENTIALS_API_LATEST, 4);
+				Credentials.Type = EOS_ELoginCredentialType::EOS_LCT_PersistentAuth;
+#if EOS_AUTH_TOKEN_SAVEGAME_STORAGE
+				auto AuthTokenUtf8 = StringCast<UTF8CHAR>(*SavedEOSAuthToken.Token);
+				Credentials.Token = (const char*)AuthTokenUtf8.Get();
+#endif
+
+				EOS_Auth_LoginOptions PALoginOptions = { };
+				PALoginOptions.ApiVersion = 3;
+				UE_EOS_CHECK_API_MISMATCH(EOS_AUTH_LOGIN_API_LATEST, 3);
+				PALoginOptions.ScopeFlags = GetAuthScopeFlags();
+				PALoginOptions.Credentials = &Credentials;
+
+				EOSHelpers->GetSystemAuthCredentialsOptions(Credentials.SystemAuthCredentialsOptions);
+
+				FLoginCallback* PACallbackObj = new FLoginCallback(AsWeak());
+				PACallbackObj->CallbackLambda = [this, LocalUserNum, AccountCredentials](const EOS_Auth_LoginCallbackInfo* Data)
+				{
+					OnEOSAuthLoginComplete(LocalUserNum, AccountCredentials, true, Data);
+				};
+
+				EOS_Auth_Login(EOSSubsystem->AuthHandle, &PALoginOptions, (void*)PACallbackObj, PACallbackObj->GetCallbackPtr());
+
+				return;
+#if EOS_AUTH_TOKEN_SAVEGAME_STORAGE
+			}
+			else
+			{
+				// The saved token is expired, we'll delete the saved token and obtain a new one
+				if (!SaveGameSystem->DeleteGame(false, *SavedEOSAuthTokenFilename, LocalUserNum))
+				{
+					UE_LOG_ONLINE(Warning, TEXT("Unable to delete Saved EOS Auth Token [%s] for LocalUserNum [%d]"), *SavedEOSAuthTokenFilename, LocalUserNum);
+				}
+			}
+		}
+		else
+		{
+			UE_LOG_ONLINE(Warning, TEXT("Unable to load Saved EOS Auth Token [%s] for LocalUserNum [%d]"), *SavedEOSAuthTokenFilename, LocalUserNum);
+		}
+	}
+
+	// If we don't have an auth token saved or it's expired, we'll get one via AccountPortal login first
+	LoginViaAccountPortal(LocalUserNum, AccountCredentials);
+#endif
+}
+
 // Chose arbitrarily since the SDK doesn't define it
 #define EOS_MAX_TOKEN_SIZE 4096
 
@@ -471,6 +573,31 @@ FLocalUserEOS& FUserManagerEOS::GetLocalUserChecked(int32 LocalUserNum)
 	return LocalUsers[LocalUserNum];
 }
 
+void FUserManagerEOS::LoginViaAccountPortal(int32 LocalUserNum, const FOnlineAccountCredentials& AccountCredentials)
+{
+	EOS_Auth_Credentials EOSCredentials = {};
+	EOSCredentials.ApiVersion = 4;
+	UE_EOS_CHECK_API_MISMATCH(EOS_AUTH_CREDENTIALS_API_LATEST, 4);
+	EOSCredentials.Type = EOS_ELoginCredentialType::EOS_LCT_AccountPortal;
+
+	EOS_Auth_LoginOptions LoginOptions = { };
+	LoginOptions.ApiVersion = 3;
+	UE_EOS_CHECK_API_MISMATCH(EOS_AUTH_LOGIN_API_LATEST, 3);
+	LoginOptions.ScopeFlags = GetAuthScopeFlags();
+	LoginOptions.Credentials = &EOSCredentials;
+
+	FPlatformEOSHelpersPtr EOSHelpers = EOSSubsystem->GetEOSHelpers();
+	EOSHelpers->GetSystemAuthCredentialsOptions(EOSCredentials.SystemAuthCredentialsOptions);
+
+	FLoginCallback* CallbackObj = new FLoginCallback(AsWeak());
+	CallbackObj->CallbackLambda = [this, LocalUserNum, AccountCredentials](const EOS_Auth_LoginCallbackInfo* Data) mutable
+	{
+		OnEOSAuthLoginComplete(LocalUserNum, AccountCredentials, true, Data);
+	};
+
+	EOS_Auth_Login(EOSSubsystem->AuthHandle, &LoginOptions, (void*)CallbackObj, CallbackObj->GetCallbackPtr());
+}
+
 void FUserManagerEOS::CallEOSAuthLogin(int32 LocalUserNum, const FOnlineAccountCredentials& Credentials)
 {
 	// First we construct the EOS Credentials object
@@ -498,9 +625,7 @@ void FUserManagerEOS::CallEOSAuthLogin(int32 LocalUserNum, const FOnlineAccountC
 	EOS_Auth_LoginOptions LoginOptions = { };
 	LoginOptions.ApiVersion = 3;
 	UE_EOS_CHECK_API_MISMATCH(EOS_AUTH_LOGIN_API_LATEST, 3);
-
 	LoginOptions.ScopeFlags = GetAuthScopeFlags();
-
 	LoginOptions.Credentials = &EOSCredentials;
 
 	FPlatformEOSHelpersPtr EOSHelpers = EOSSubsystem->GetEOSHelpers();
@@ -513,55 +638,130 @@ void FUserManagerEOS::CallEOSAuthLogin(int32 LocalUserNum, const FOnlineAccountC
 	FLoginCallback* CallbackObj = new FLoginCallback(AsWeak());
 	CallbackObj->CallbackLambda = [this, LocalUserNum, Credentials, bIsPersistentLogin](const EOS_Auth_LoginCallbackInfo* Data)
 	{
-		if (Data->ResultCode == EOS_EResult::EOS_Success)
+		OnEOSAuthLoginComplete(LocalUserNum, Credentials, bIsPersistentLogin, Data);
+	};
+
+	// Perform the auth call
+	EOS_Auth_Login(EOSSubsystem->AuthHandle, &LoginOptions, (void*)CallbackObj, CallbackObj->GetCallbackPtr());
+}
+
+void FUserManagerEOS::CopyAndSaveEpicAuthToken(int32 LocalUserNum, const EOS_EpicAccountId& EpicAccountId)
+{
+	EOS_Auth_Token* AuthToken = nullptr;
+	EOS_Auth_CopyUserAuthTokenOptions CopyOptions = { };
+	CopyOptions.ApiVersion = 1;
+	UE_EOS_CHECK_API_MISMATCH(EOS_AUTH_COPYUSERAUTHTOKEN_API_LATEST, 1);
+
+	EOS_EResult CopyResult = EOS_Auth_CopyUserAuthToken(EOSSubsystem->AuthHandle, &CopyOptions, EpicAccountId, &AuthToken);
+	if (CopyResult == EOS_EResult::EOS_Success)
+	{
+		FSavedEOSAuthToken SavedEOSAuthToken;
+		SavedEOSAuthToken.Token = StringCast<TCHAR>(AuthToken->RefreshToken).Get();
+		const FString RefreshExpiresAtTCHARPtr = StringCast<TCHAR>(AuthToken->RefreshExpiresAt).Get();
+		if (FDateTime::ParseIso8601(*RefreshExpiresAtTCHARPtr, SavedEOSAuthToken.ExpirationDate))
 		{
-			// Continue the login process by getting the product user id for EAS only
-			ConnectLoginEAS(LocalUserNum, Data->LocalUserId, Credentials);
+			FString SavedEOSAuthTokenStr = SavedEOSAuthToken.ToJson(false);
+
+			TArray<uint8> SavedEOSAuthTokenData;
+			SavedEOSAuthTokenData.SetNumUninitialized(SavedEOSAuthTokenStr.Len());
+			StringToBytes(SavedEOSAuthTokenStr, SavedEOSAuthTokenData.GetData(), SavedEOSAuthTokenStr.Len());
+
+			ISaveGameSystem* SaveGameSystem = IPlatformFeaturesModule::Get().GetSaveGameSystem();
+
+			FString SavedEOSAuthTokenFilename = GetEOSAuthTokenFilename();
+
+			if (SaveGameSystem->SaveGame(false, *SavedEOSAuthTokenFilename, LocalUserNum, SavedEOSAuthTokenData))
+			{
+				UE_LOG_ONLINE(Verbose, TEXT("Saved EOS Auth Token to File [%s]. Contents: [%s]"), *SavedEOSAuthTokenFilename, UE_BUILD_SHIPPING ? TEXT("<redacted>") : *SavedEOSAuthTokenStr);
+			}
+			else
+			{
+				UE_LOG_ONLINE(Warning, TEXT("Unable to save EOS Auth Token [%s] for LocalUserNum [%d]"), *SavedEOSAuthTokenFilename, LocalUserNum);
+			}
 		}
-		else if (Data->ResultCode == EOS_EResult::EOS_InvalidUser)
+		else
 		{
-			// Link the account
-			LinkEAS(LocalUserNum, Data->ContinuanceToken, Credentials);
+			UE_LOG_ONLINE(Warning, TEXT("Unable to parse RefreshExpiresAt for FSavedEOSAuthToken. EOS Auth Token for user [%d] was not saved"), LocalUserNum);
 		}
-		else
+	}
+	else
+	{
+		UE_LOG_ONLINE(Warning, TEXT("EOS_Auth_CopyUserAuthToken failed with result [%s]"), *LexToString(CopyResult));
+	}
+}
+
+void FUserManagerEOS::OnEOSAuthLoginComplete(int32 LocalUserNum, const FOnlineAccountCredentials& Credentials, bool bIsPersistentLogin, const EOS_Auth_LoginCallbackInfo* Data)
+{
+	if (Data->ResultCode == EOS_EResult::EOS_Success)
+	{
+		if (bIsPersistentLogin)
+		{
+#if EOS_AUTH_TOKEN_SAVEGAME_STORAGE
+			CopyAndSaveEpicAuthToken(LocalUserNum, Data->LocalUserId);
+#endif
+			EOSSubsystem->GetEOSHelpers()->AddExternalAccountMapping(EOSSubsystem->GetEOSPlatformHandle(), Data->LocalUserId, LocalUserNum);
+		}
+
+		// Continue the login process by getting the product user id for EAS only
+		ConnectLoginEAS(LocalUserNum, Data->LocalUserId, Credentials);
+	}
+	else if (Data->ResultCode == EOS_EResult::EOS_InvalidUser)
+	{
+		// Link the account
+		LinkEAS(LocalUserNum, Data->ContinuanceToken, Credentials);
+	}
+	else if (Data->ResultCode == EOS_EResult::EOS_InvalidAuth && bIsPersistentLogin)
+	{
+		// We attempted a Persistent Auth login but there is no stored token, we'll attempt Account Portal login automatically
+
+		LoginViaAccountPortal(LocalUserNum, Credentials);
+	}
+	else
+	{
+		auto TriggerLoginFailure = [this, LocalUserNum, LoginResultCode = Data->ResultCode]()
 		{
-			auto TriggerLoginFailure = [this, LocalUserNum, LoginResultCode = Data->ResultCode]()
+			FString ErrorString = FString::Printf(TEXT("Login(%d) failed with EOS result code (%s)"), LocalUserNum, ANSI_TO_TCHAR(EOS_EResult_ToString(LoginResultCode)));
+			UE_LOG_ONLINE(Warning, TEXT("%s"), *ErrorString);
+			TriggerOnLoginCompleteDelegates(LocalUserNum, false, *FUniqueNetIdEOS::EmptyId(), ErrorString);
+		};
+
+		const bool bShouldRemoveCachedToken =
+			Data->ResultCode == EOS_EResult::EOS_AccessDenied ||
+			Data->ResultCode == EOS_EResult::EOS_Auth_InvalidToken;
+
+		// Check for invalid persistent login credentials.
+		if (bIsPersistentLogin && bShouldRemoveCachedToken)
+		{
+			FDeletePersistentAuthCallback* DeleteAuthCallbackObj = new FDeletePersistentAuthCallback(AsWeak());
+			DeleteAuthCallbackObj->CallbackLambda = [this, LocalUserNum, TriggerLoginFailure](const EOS_Auth_DeletePersistentAuthCallbackInfo* Data)
 			{
-				FString ErrorString = FString::Printf(TEXT("Login(%d) failed with EOS result code (%s)"), LocalUserNum, ANSI_TO_TCHAR(EOS_EResult_ToString(LoginResultCode)));
-				UE_LOG_ONLINE(Warning, TEXT("%s"), *ErrorString);
-				TriggerOnLoginCompleteDelegates(LocalUserNum, false, *FUniqueNetIdEOS::EmptyId(), ErrorString);
+				// Deleting the auth token is best effort.
+				TriggerLoginFailure();
 			};
 
-			const bool bShouldRemoveCachedToken =
-				Data->ResultCode == EOS_EResult::EOS_InvalidAuth ||
-				Data->ResultCode == EOS_EResult::EOS_AccessDenied ||
-				Data->ResultCode == EOS_EResult::EOS_Auth_InvalidToken;
+			EOS_Auth_DeletePersistentAuthOptions DeletePersistentAuthOptions;
+			DeletePersistentAuthOptions.ApiVersion = 2;
+			UE_EOS_CHECK_API_MISMATCH(EOS_AUTH_DELETEPERSISTENTAUTH_API_LATEST, 2);
+			DeletePersistentAuthOptions.RefreshToken = nullptr;
+			EOS_Auth_DeletePersistentAuth(EOSSubsystem->AuthHandle, &DeletePersistentAuthOptions, (void*)DeleteAuthCallbackObj, DeleteAuthCallbackObj->GetCallbackPtr());
 
-			// Check for invalid persistent login credentials.
-			if (bIsPersistentLogin && bShouldRemoveCachedToken)
-			{
-				FDeletePersistentAuthCallback* DeleteAuthCallbackObj = new FDeletePersistentAuthCallback(AsWeak());
-				DeleteAuthCallbackObj->CallbackLambda = [this, LocalUserNum, TriggerLoginFailure](const EOS_Auth_DeletePersistentAuthCallbackInfo* Data)
-				{
-					// Deleting the auth token is best effort.
-					TriggerLoginFailure();
-				};
+#if EOS_AUTH_TOKEN_SAVEGAME_STORAGE
+			// After deleting the token at API level, we'll also delete any saved tokens for this user
+			ISaveGameSystem* SaveGameSystem = IPlatformFeaturesModule::Get().GetSaveGameSystem();
 
-				EOS_Auth_DeletePersistentAuthOptions DeletePersistentAuthOptions;
-				DeletePersistentAuthOptions.ApiVersion = 2;
-				UE_EOS_CHECK_API_MISMATCH(EOS_AUTH_DELETEPERSISTENTAUTH_API_LATEST, 2);
-				DeletePersistentAuthOptions.RefreshToken = nullptr;
-				EOS_Auth_DeletePersistentAuth(EOSSubsystem->AuthHandle, &DeletePersistentAuthOptions, (void*)DeleteAuthCallbackObj, DeleteAuthCallbackObj->GetCallbackPtr());
-			}
-			else
+			FString SavedEOSAuthTokenFilename = GetEOSAuthTokenFilename();
+
+			if (!SaveGameSystem->DeleteGame(false, *SavedEOSAuthTokenFilename, LocalUserNum))
 			{
-				TriggerLoginFailure();
+				UE_LOG_ONLINE(Warning, TEXT("Unable to delete Saved EOS Auth Token [%s] for LocalUserNum [%d]"), *SavedEOSAuthTokenFilename, LocalUserNum);
 			}
+#endif
 		}
-	};
-
-	// Perform the auth call
-	EOS_Auth_Login(EOSSubsystem->AuthHandle, &LoginOptions, (void*)CallbackObj, CallbackObj->GetCallbackPtr());
+		else
+		{
+			TriggerLoginFailure();
+		}
+	}
 }
 
 struct FLinkAccountOptions :
@@ -705,53 +905,86 @@ bool FUserManagerEOS::ConnectLoginNoEAS(int32 LocalUserNum, const FOnlineAccount
 
 bool FUserManagerEOS::ConnectLoginEAS(int32 LocalUserNum, EOS_EpicAccountId AccountId, const FOnlineAccountCredentials& AccountCredentials)
 {
-	EOS_Auth_Token* AuthToken = nullptr;
-	EOS_Auth_CopyUserAuthTokenOptions CopyOptions = { };
-	CopyOptions.ApiVersion = 1;
-	UE_EOS_CHECK_API_MISMATCH(EOS_AUTH_COPYUSERAUTHTOKEN_API_LATEST, 1);
+#if ADD_USER_LOGIN_INFO
+	GetPlatformAuthToken(LocalUserNum,
+		FOnGetLinkedAccountAuthTokenCompleteDelegate::CreateLambda([this, WeakThis = AsWeak(), AccountId, AccountCredentials](int32 LocalUserNum, bool bWasSuccessful, const FExternalAuthToken& AuthToken)
+			{
+				if (FUserManagerEOSPtr StrongThis = WeakThis.Pin())
+				{
+					if (!bWasSuccessful || !AuthToken.IsValid())
+					{
+						const FString ErrorString = FString::Printf(TEXT("ConnectLoginEAS(%d) failed due to the platform OSS giving an empty auth token"), LocalUserNum);
+						UE_LOG_ONLINE(Warning, TEXT("%s"), *ErrorString);
+						TriggerOnLoginCompleteDelegates(LocalUserNum, false, *FUniqueNetIdEOS::EmptyId(), ErrorString);
+						return;
+					}
+#endif
 
-	EOS_EResult CopyResult = EOS_Auth_CopyUserAuthToken(EOSSubsystem->AuthHandle, &CopyOptions, AccountId, &AuthToken);
-	if (CopyResult == EOS_EResult::EOS_Success)
-	{
-		EOS_Connect_Credentials Credentials = { };
-		Credentials.ApiVersion = 1;
-		UE_EOS_CHECK_API_MISMATCH(EOS_CONNECT_CREDENTIALS_API_LATEST, 1);
-		Credentials.Type = EOS_EExternalCredentialType::EOS_ECT_EPIC;
-		Credentials.Token = AuthToken->AccessToken;
+					EOS_Auth_Token* EOSAuthToken = nullptr;
+					EOS_Auth_CopyUserAuthTokenOptions CopyOptions = { };
+					CopyOptions.ApiVersion = 1;
+					UE_EOS_CHECK_API_MISMATCH(EOS_AUTH_COPYUSERAUTHTOKEN_API_LATEST, 1);
 
-		EOS_Connect_LoginOptions Options = { };
-		Options.ApiVersion = 2;
-		UE_EOS_CHECK_API_MISMATCH(EOS_CONNECT_LOGIN_API_LATEST, 2);
-		Options.Credentials = &Credentials;
+					EOS_EResult CopyResult = EOS_Auth_CopyUserAuthToken(EOSSubsystem->AuthHandle, &CopyOptions, AccountId, &EOSAuthToken);
+					if (CopyResult == EOS_EResult::EOS_Success)
+					{
+						EOS_Connect_Credentials Credentials = { };
+						Credentials.ApiVersion = 1;
+						UE_EOS_CHECK_API_MISMATCH(EOS_CONNECT_CREDENTIALS_API_LATEST, 1);
+						Credentials.Type = EOS_EExternalCredentialType::EOS_ECT_EPIC;
+						Credentials.Token = EOSAuthToken->AccessToken;
 
-		FConnectLoginCallback* CallbackObj = new FConnectLoginCallback(AsWeak());
-		CallbackObj->CallbackLambda = [LocalUserNum, AccountId, AccountCredentials, this](const EOS_Connect_LoginCallbackInfo* Data)
-		{
-			if (Data->ResultCode == EOS_EResult::EOS_Success)
-			{
-				// We have an account mapping, skip to final login
-				FullLoginCallback(LocalUserNum, AccountId, Data->LocalUserId, AccountCredentials);
-			}
-			else if (Data->ResultCode == EOS_EResult::EOS_InvalidUser)
-			{
-				// We need to create the mapping for this user using the continuation token
-				CreateConnectedLogin(LocalUserNum, AccountId, Data->ContinuanceToken, AccountCredentials);
-			}
-			else
-			{
-				UE_LOG_ONLINE(Error, TEXT("ConnectLogin(%d) failed with EOS result code (%s)"), LocalUserNum, ANSI_TO_TCHAR(EOS_EResult_ToString(Data->ResultCode)));
-				Logout(LocalUserNum);
-			}
-		};
-		EOS_Connect_Login(EOSSubsystem->ConnectHandle, &Options, CallbackObj, CallbackObj->GetCallbackPtr());
+#if ADD_USER_LOGIN_INFO
+						EOS_Connect_UserLoginInfo UserLoginInfo = {};
+						UserLoginInfo.ApiVersion = 2;
+						UE_EOS_CHECK_API_MISMATCH(EOS_CONNECT_USERLOGININFO_API_LATEST, 2);
+						auto AuthTokenConverter = StringCast<UTF8CHAR>(*AuthToken.TokenString);
+						UserLoginInfo.NsaIdToken = (const char*)AuthTokenConverter.Get();
+#endif
+
+						EOS_Connect_LoginOptions Options = { };
+						Options.ApiVersion = 2;
+						UE_EOS_CHECK_API_MISMATCH(EOS_CONNECT_LOGIN_API_LATEST, 2);
+						Options.Credentials = &Credentials;
+#if ADD_USER_LOGIN_INFO
+						Options.UserLoginInfo = &UserLoginInfo;
+#else
+						Options.UserLoginInfo = nullptr;
+#endif
+
+						FConnectLoginCallback* CallbackObj = new FConnectLoginCallback(AsWeak());
+						CallbackObj->CallbackLambda = [LocalUserNum, AccountId, AccountCredentials, this](const EOS_Connect_LoginCallbackInfo* Data)
+						{
+							if (Data->ResultCode == EOS_EResult::EOS_Success)
+							{
+								// We have an account mapping, skip to final login
+								FullLoginCallback(LocalUserNum, AccountId, Data->LocalUserId, AccountCredentials);
+							}
+							else if (Data->ResultCode == EOS_EResult::EOS_InvalidUser)
+							{
+								// We need to create the mapping for this user using the continuation token
+								CreateConnectedLogin(LocalUserNum, AccountId, Data->ContinuanceToken, AccountCredentials);
+							}
+							else
+							{
+								UE_LOG_ONLINE(Error, TEXT("ConnectLogin(%d) failed with EOS result code (%s)"), LocalUserNum, ANSI_TO_TCHAR(EOS_EResult_ToString(Data->ResultCode)));
+								Logout(LocalUserNum);
+							}
+						};
+						EOS_Connect_Login(EOSSubsystem->ConnectHandle, &Options, CallbackObj, CallbackObj->GetCallbackPtr());
+
+						EOS_Auth_Token_Release(EOSAuthToken);
+					}
+					else
+					{
+						UE_LOG_ONLINE(Error, TEXT("ConnectLogin(%d) failed with EOS result code (%s)"), LocalUserNum, ANSI_TO_TCHAR(EOS_EResult_ToString(CopyResult)));
+						Logout(LocalUserNum);
+					}
+#if ADD_USER_LOGIN_INFO
+				}
+			}));
+#endif
 
-		EOS_Auth_Token_Release(AuthToken);
-	}
-	else
-	{
-		UE_LOG_ONLINE(Error, TEXT("ConnectLogin(%d) failed with EOS result code (%s)"), LocalUserNum, ANSI_TO_TCHAR(EOS_EResult_ToString(CopyResult)));
-		Logout(LocalUserNum);
-	}
 	return true;
 }
 
@@ -1617,25 +1850,29 @@ ELoginStatus::Type FUserManagerEOS::GetLoginStatus(const FUniqueNetId& UserId) c
 
 FString FUserManagerEOS::GetPlayerNickname(int32 LocalUserNum) const
 {
-	FUniqueNetIdEOSPtr UserId = GetLocalUniqueNetIdEOS(LocalUserNum);
+	const FUniqueNetIdEOSPtr UserId = GetLocalUniqueNetIdEOS(LocalUserNum);
 	if (UserId.IsValid())
 	{
-		TSharedPtr<FUserOnlineAccount> UserAccount = GetUserAccount(*UserId);
-		if (UserAccount.IsValid())
-		{
-			return UserAccount->GetDisplayName();
-		}
+		return GetPlayerNickname(*UserId);
 	}
+
 	return FString();
 }
 
 FString FUserManagerEOS::GetPlayerNickname(const FUniqueNetId& UserId) const
 {
-	TSharedPtr<FUserOnlineAccount> UserAccount = GetUserAccount(UserId);
-	if (UserAccount.IsValid())
+	// GetUserAccount only searches local users
+	const TSharedPtr<FUserOnlineAccount> LocalUserAccount = GetUserAccount(UserId);
+	if (LocalUserAccount.IsValid())
+	{
+		return LocalUserAccount->GetDisplayName();
+	}
+
+	if (const FOnlineUserEOSRef* RemoteUserRef = UniqueNetIdToUserRefMap.Find(UserId.AsShared()))
 	{
-		return UserAccount->GetDisplayName();
+		return RemoteUserRef->Get().GetDisplayName();
 	}
+
 	return FString();
 }
 
diff --git a/Engine/Plugins/Online/OnlineSubsystemEOS/Source/OnlineSubsystemEOS/Private/UserManagerEOS.h b/Engine/Plugins/Online/OnlineSubsystemEOS/Source/OnlineSubsystemEOS/Private/UserManagerEOS.h
--- a/Engine/Plugins/Online/OnlineSubsystemEOS/Source/OnlineSubsystemEOS/Private/UserManagerEOS.h
+++ b/Engine/Plugins/Online/OnlineSubsystemEOS/Source/OnlineSubsystemEOS/Private/UserManagerEOS.h
@@ -379,7 +379,9 @@ public:
 	FUserManagerEOS() = delete;
 
 	bool ConnectLoginEAS(int32 LocalUserNum, EOS_EpicAccountId AccountId, const FOnlineAccountCredentials& AccountCredentials);
+	void LoginViaPersistentAuth(int32 LocalUserNum, const FOnlineAccountCredentials& CurrentLoginCredentials);
 	void LoginViaExternalAuth(int32 LocalUserNum, const FOnlineAccountCredentials& AccountCredentials);
+	void LoginViaAccountPortal(int32 LocalUserNum, const FOnlineAccountCredentials& AccountCredentials);
 	void CreateConnectedLogin(int32 LocalUserNum, EOS_EpicAccountId AccountId, EOS_ContinuanceToken Token, const FOnlineAccountCredentials& AccountCredentials);
 	void LinkEAS(int32 LocalUserNum, EOS_ContinuanceToken Token, const FOnlineAccountCredentials& AccountCredentials);
 	void RefreshConnectLogin(int32 LocalUserNum);
@@ -396,6 +398,8 @@ private:
 	FLocalUserEOS& GetLocalUserChecked(int32 LocalUserNum);
 
 	void CallEOSAuthLogin(int32 LocalUserNum, const FOnlineAccountCredentials& Credentials);
+	void CopyAndSaveEpicAuthToken(int32 LocalUserNum, const EOS_EpicAccountId& EpicAccountId);
+	void OnEOSAuthLoginComplete(int32 LocalUserNum, const FOnlineAccountCredentials& Credentials, bool bIsPersistentLogin, const EOS_Auth_LoginCallbackInfo* Data);
 
 	void RemoveLocalUser(int32 LocalUserNum);
 	FLocalUserEOS& AddLocalUser(int32 LocalUserNum, EOS_EpicAccountId EpicAccountId, EOS_ProductUserId UserId, const FOnlineAccountCredentials& AccountCredentials);
@@ -419,6 +423,9 @@ private:
 	void GetPlatformAuthToken(int32 LocalUserNum, const FOnGetLinkedAccountAuthTokenCompleteDelegate& Delegate) const;
 	FString GetPlatformDisplayName(int32 LocalUserNum) const;
 
+	/** Get the file name for the Epic Auth Token file used in Persistent Auth */
+	FString GetEOSAuthTokenFilename();
+
 	/** Cached pointer to owning subsystem */
 	FOnlineSubsystemEOS* EOSSubsystem;
 
diff --git a/Engine/Plugins/Online/OnlineSubsystemEOS/Source/OnlineSubsystemEOS/Public/EOSSettings.h b/Engine/Plugins/Online/OnlineSubsystemEOS/Source/OnlineSubsystemEOS/Public/EOSSettings.h
--- a/Engine/Plugins/Online/OnlineSubsystemEOS/Source/OnlineSubsystemEOS/Public/EOSSettings.h
+++ b/Engine/Plugins/Online/OnlineSubsystemEOS/Source/OnlineSubsystemEOS/Public/EOSSettings.h
@@ -77,6 +77,7 @@ struct FEOSSettings
 	bool bEnableOverlay;
 	bool bEnableSocialOverlay;
 	bool bEnableEditorOverlay;
+	bool bPreferPersistentAuth;
 	bool bUseEAS;
 	bool bUseEOSConnect;
 	bool bUseEOSSessions;
@@ -127,6 +128,10 @@ public:
 	UPROPERTY(Config, EditAnywhere, BlueprintReadOnly, Category = "EOS Settings")
 	bool bEnableEditorOverlay = false;
 
+	/** Set to true to prefer persistent auth over external authentication during Login */
+	UPROPERTY(Config, EditAnywhere, BlueprintReadOnly, Category="EOS Settings")
+	bool bPreferPersistentAuth = false;
+
 	/** Tag combinations for paged queries in title file enumerations, separate tags within groups using `+` */
 	UPROPERTY(Config, EditAnywhere, BlueprintReadOnly, Category="EOS Settings")
 	TArray<FString> TitleStorageTags;
