Added logic to test EOS Leaderboard methods via commandline

diff --git a/Engine/Plugins/Online/OnlineSubsystem/Source/Private/OnlineStats.cpp b/Engine/Plugins/Online/OnlineSubsystem/Source/Private/OnlineStats.cpp
--- a/Engine/Plugins/Online/OnlineSubsystem/Source/Private/OnlineStats.cpp
+++ b/Engine/Plugins/Online/OnlineSubsystem/Source/Private/OnlineStats.cpp
@@ -146,3 +146,27 @@ void FOnlineStats::DecrementIntStat(const FName& StatName, int32 DecBy)
 		Properties.Add(StatName, NewValue);
 	}
 }
+
+FString FOnlineStatsRow::ToLogString() const
+{
+	FString LogString = FString::Printf(TEXT("%d : %s"), Rank, *NickName);
+
+	for (const TPair<FName, FVariantData>& Column : Columns)
+	{
+		LogString += FString::Printf(TEXT("\t\t%s : %s"), *Column.Key.ToString(), *Column.Value.ToString());
+	}
+
+	return LogString;
+}
+
+FString FOnlineLeaderboardRead::ToLogString() const
+{
+	FString LogString = FString::Printf(TEXT("\nLeaderboardName: %s\nSortedColumn: %s\nRows:\n"), *LeaderboardName.ToString(), *SortedColumn.ToString());
+
+	for (const FOnlineStatsRow& Row : Rows)
+	{
+		LogString += FString::Printf(TEXT("\t%s\n"), *Row.ToLogString());
+	}
+
+	return LogString;
+}
\ No newline at end of file
diff --git a/Engine/Plugins/Online/OnlineSubsystem/Source/Public/OnlineStats.h b/Engine/Plugins/Online/OnlineSubsystem/Source/Public/OnlineStats.h
--- a/Engine/Plugins/Online/OnlineSubsystem/Source/Public/OnlineStats.h
+++ b/Engine/Plugins/Online/OnlineSubsystem/Source/Public/OnlineStats.h
@@ -143,6 +143,8 @@ public:
 		PlayerId(InPlayerId)
 	{
 	}
+
+	FString ToLogString() const;
 };
 
 /**
@@ -212,6 +214,8 @@ public:
 
 		return NULL;
 	}
+
+	FString ToLogString() const;
 };
 
 typedef TSharedRef<FOnlineLeaderboardRead, ESPMode::ThreadSafe> FOnlineLeaderboardReadRef;
diff --git a/Engine/Plugins/Online/OnlineSubsystemEOS/Source/OnlineSubsystemEOS/Private/OnlineLeaderboardsEOS.cpp b/Engine/Plugins/Online/OnlineSubsystemEOS/Source/OnlineSubsystemEOS/Private/OnlineLeaderboardsEOS.cpp
--- a/Engine/Plugins/Online/OnlineSubsystemEOS/Source/OnlineSubsystemEOS/Private/OnlineLeaderboardsEOS.cpp
+++ b/Engine/Plugins/Online/OnlineSubsystemEOS/Source/OnlineSubsystemEOS/Private/OnlineLeaderboardsEOS.cpp
@@ -380,4 +380,131 @@ bool FOnlineLeaderboardsEOS::WriteOnlinePlayerRatings(const FName& SessionName,
 	return false;
 }
 
+bool FOnlineLeaderboardsEOS::HandleLeaderboardsExec(UWorld* InWorld, const TCHAR* Cmd, FOutputDevice& Ar)
+{
+#if !UE_BUILD_SHIPPING
+	bool bWasHandled = true;
+
+	if (FParse::Command(&Cmd, TEXT("ReadLeaderboards"))) /* ONLINE (EOS if using EOSPlus) LEADERBOARDS ReadLeaderboards LocalUserNum=0 LeaderboardName=Deaths SortedColumn=Deaths DataType=Int32 */
+	{
+		int LocalUserNum;
+		FParse::Value(Cmd, TEXT("LocalUserNum="), LocalUserNum);
+		const FUniqueNetIdPtr NetIdPtr = EOSSubsystem->UserManager->GetUniquePlayerId(LocalUserNum);
+
+		FString LeaderboardName;
+		FParse::Value(Cmd, TEXT("LeaderboardName="), LeaderboardName);
+		FString SortedColumn;
+		FParse::Value(Cmd, TEXT("SortedColumn="), SortedColumn);
+		FString DataTypeStr;
+		FParse::Value(Cmd, TEXT("DataType="), DataTypeStr);
+		const EOnlineKeyValuePairDataType::Type DataType = EOnlineKeyValuePairDataType::FromString(DataTypeStr);
+
+		FOnlineLeaderboardReadRef ReadRef = MakeShared<FOnlineLeaderboardRead>();
+		ReadRef->LeaderboardName = FName(*LeaderboardName);
+		ReadRef->SortedColumn = FName(*SortedColumn);
+		ReadRef->ColumnMetadata.Add(FColumnMetaData(FName(*SortedColumn), DataType));
+
+		AddOnLeaderboardReadCompleteDelegate_Handle(FOnLeaderboardReadCompleteDelegate::CreateLambda([ReadRef](bool bWasSuccessful)
+			{
+				UE_LOG_ONLINE_LEADERBOARD(Log, TEXT("FOnlineLeaderboardsEOS::ReadLeaderboards finished with bWasSuccessful=%s"), *LexToString(bWasSuccessful));
+				UE_LOG_ONLINE_LEADERBOARD(Log, TEXT("%s"), *ReadRef->ToLogString());
+			}));
+
+		ReadLeaderboards({ NetIdPtr.ToSharedRef() }, ReadRef);
+	}
+	else if (FParse::Command(&Cmd, TEXT("ReadLeaderboardsForFriends"))) /* ONLINE (EOS if using EOSPlus) LEADERBOARDS ReadLeaderboardsForFriends LocalUserNum=0 LeaderboardName=Deaths SortedColumn=Deaths DataType=Int32 */
+	{
+		int LocalUserNum;
+		FParse::Value(Cmd, TEXT("LocalUserNum="), LocalUserNum);
+
+		FString LeaderboardName;
+		FParse::Value(Cmd, TEXT("LeaderboardName="), LeaderboardName);
+		FString SortedColumn;
+		FParse::Value(Cmd, TEXT("SortedColumn="), SortedColumn);
+		FString DataTypeStr;
+		FParse::Value(Cmd, TEXT("DataType="), DataTypeStr);
+		const EOnlineKeyValuePairDataType::Type DataType = EOnlineKeyValuePairDataType::FromString(DataTypeStr);
+
+		FOnlineLeaderboardReadRef ReadRef = MakeShared<FOnlineLeaderboardRead>();
+		ReadRef->LeaderboardName = FName(*LeaderboardName);
+		ReadRef->SortedColumn = FName(*SortedColumn);
+		ReadRef->ColumnMetadata.Add(FColumnMetaData(FName(*SortedColumn), DataType));
+
+		AddOnLeaderboardReadCompleteDelegate_Handle(FOnLeaderboardReadCompleteDelegate::CreateLambda([ReadRef](bool bWasSuccessful)
+			{
+				UE_LOG_ONLINE_LEADERBOARD(Log, TEXT("FOnlineLeaderboardsEOS::ReadLeaderboards finished with bWasSuccessful=%s"), *LexToString(bWasSuccessful));
+				UE_LOG_ONLINE_LEADERBOARD(Log, TEXT("%s"), *ReadRef->ToLogString());
+			}));
+
+		ReadLeaderboardsForFriends(LocalUserNum, ReadRef);
+	}
+	else if (FParse::Command(&Cmd, TEXT("ReadLeaderboardsAroundRank"))) /* ONLINE (EOS if using EOSPlus) LEADERBOARDS ReadLeaderboardsAroundRank Rank=1 Range=10 LeaderboardName=Deaths SortedColumn=Deaths DataType=Int32 */
+	{
+		int Rank;
+		FParse::Value(Cmd, TEXT("Rank="), Rank);
+		int Range;
+		FParse::Value(Cmd, TEXT("Range="), Range);
+
+		FString LeaderboardName;
+		FParse::Value(Cmd, TEXT("LeaderboardName="), LeaderboardName);
+		FString SortedColumn;
+		FParse::Value(Cmd, TEXT("SortedColumn="), SortedColumn);
+		FString DataTypeStr;
+		FParse::Value(Cmd, TEXT("DataType="), DataTypeStr);
+		const EOnlineKeyValuePairDataType::Type DataType = EOnlineKeyValuePairDataType::FromString(DataTypeStr);
+
+		FOnlineLeaderboardReadRef ReadRef = MakeShared<FOnlineLeaderboardRead>();
+		ReadRef->LeaderboardName = FName(*LeaderboardName);
+		ReadRef->SortedColumn = FName(*SortedColumn);
+		ReadRef->ColumnMetadata.Add(FColumnMetaData(FName(*SortedColumn), DataType));
+
+		AddOnLeaderboardReadCompleteDelegate_Handle(FOnLeaderboardReadCompleteDelegate::CreateLambda([ReadRef](bool bWasSuccessful)
+			{
+				UE_LOG_ONLINE_LEADERBOARD(Log, TEXT("FOnlineLeaderboardsEOS::ReadLeaderboards finished with bWasSuccessful=%s"), *LexToString(bWasSuccessful));
+				UE_LOG_ONLINE_LEADERBOARD(Log, TEXT("%s"), *ReadRef->ToLogString());
+			}));
+
+		ReadLeaderboardsAroundRank(Rank, Range, ReadRef);
+	}
+	else if (FParse::Command(&Cmd, TEXT("ReadLeaderboardsAroundUser"))) /* ONLINE (EOS if using EOSPlus) LEADERBOARDS ReadLeaderboardsAroundUser LocalUserNum=0 Range=10 LeaderboardName=Deaths SortedColumn=Deaths DataType=Int32 */
+	{
+		int LocalUserNum;
+		FParse::Value(Cmd, TEXT("LocalUserNum="), LocalUserNum);
+		const FUniqueNetIdPtr NetIdPtr = EOSSubsystem->UserManager->GetUniquePlayerId(LocalUserNum);
+
+		int Range;
+		FParse::Value(Cmd, TEXT("Range="), Range);
+
+		FString LeaderboardName;
+		FParse::Value(Cmd, TEXT("LeaderboardName="), LeaderboardName);
+		FString SortedColumn;
+		FParse::Value(Cmd, TEXT("SortedColumn="), SortedColumn);
+		FString DataTypeStr;
+		FParse::Value(Cmd, TEXT("DataType="), DataTypeStr);
+		const EOnlineKeyValuePairDataType::Type DataType = EOnlineKeyValuePairDataType::FromString(DataTypeStr);
+
+		FOnlineLeaderboardReadRef ReadRef = MakeShared<FOnlineLeaderboardRead>();
+		ReadRef->LeaderboardName = FName(*LeaderboardName);
+		ReadRef->SortedColumn = FName(*SortedColumn);
+		ReadRef->ColumnMetadata.Add(FColumnMetaData(FName(*SortedColumn), DataType));
+
+		AddOnLeaderboardReadCompleteDelegate_Handle(FOnLeaderboardReadCompleteDelegate::CreateLambda([ReadRef](bool bWasSuccessful)
+			{
+				UE_LOG_ONLINE_LEADERBOARD(Log, TEXT("FOnlineLeaderboardsEOS::ReadLeaderboards finished with bWasSuccessful=%s"), *LexToString(bWasSuccessful));
+				UE_LOG_ONLINE_LEADERBOARD(Log, TEXT("%s"), *ReadRef->ToLogString());
+			}));
+
+		ReadLeaderboardsAroundUser(NetIdPtr.ToSharedRef(), Range, ReadRef);
+	}
+	else
+	{
+		bWasHandled = false;
+	}
+
+	return bWasHandled;
+#else
+	return false;
+#endif // !UE_BUILD_SHIPPING
+}
+
 #endif
diff --git a/Engine/Plugins/Online/OnlineSubsystemEOS/Source/OnlineSubsystemEOS/Private/OnlineLeaderboardsEOS.h b/Engine/Plugins/Online/OnlineSubsystemEOS/Source/OnlineSubsystemEOS/Private/OnlineLeaderboardsEOS.h
--- a/Engine/Plugins/Online/OnlineSubsystemEOS/Source/OnlineSubsystemEOS/Private/OnlineLeaderboardsEOS.h
+++ b/Engine/Plugins/Online/OnlineSubsystemEOS/Source/OnlineSubsystemEOS/Private/OnlineLeaderboardsEOS.h
@@ -39,6 +39,8 @@ public:
 	{
 	}
 
+	bool HandleLeaderboardsExec(UWorld* InWorld, const TCHAR* Cmd, FOutputDevice& Ar);
+
 private:
 
 	/** Reference to the main EOS subsystem */
diff --git a/Engine/Plugins/Online/OnlineSubsystemEOS/Source/OnlineSubsystemEOS/Private/OnlineSubsystemEOS.cpp b/Engine/Plugins/Online/OnlineSubsystemEOS/Source/OnlineSubsystemEOS/Private/OnlineSubsystemEOS.cpp
--- a/Engine/Plugins/Online/OnlineSubsystemEOS/Source/OnlineSubsystemEOS/Private/OnlineSubsystemEOS.cpp
+++ b/Engine/Plugins/Online/OnlineSubsystemEOS/Source/OnlineSubsystemEOS/Private/OnlineSubsystemEOS.cpp
@@ -497,19 +497,23 @@ bool FOnlineSubsystemEOS::Exec(UWorld* InWorld, const TCHAR* Cmd, FOutputDevice&
 	}
 
 	bool bWasHandled = false;
-	if (UserManager != nullptr && FParse::Command(&Cmd, TEXT("FRIENDS")))
+	if (UserManager != nullptr && FParse::Command(&Cmd, TEXT("FRIENDS"))) /* ONLINE (EOS if using EOSPlus) FRIENDS ... */
 	{
 		bWasHandled = UserManager->HandleFriendsExec(InWorld, Cmd, Ar);
 	}
-	else if (StoreInterfacePtr != nullptr && FParse::Command(&Cmd, TEXT("ECOM")))
+	else if (StoreInterfacePtr != nullptr && FParse::Command(&Cmd, TEXT("ECOM"))) /* ONLINE (EOS if using EOSPlus) ECOM ... */
 	{
 		bWasHandled = StoreInterfacePtr->HandleEcomExec(InWorld, Cmd, Ar);
 	}
-	else if (TitleFileInterfacePtr != nullptr && FParse::Command(&Cmd, TEXT("TITLEFILE")))
+	else if (LeaderboardsInterfacePtr != nullptr && FParse::Command(&Cmd, TEXT("LEADERBOARDS"))) /* ONLINE (EOS if using EOSPlus) LEADERBOARDS ... */
+	{
+		bWasHandled = LeaderboardsInterfacePtr->HandleLeaderboardsExec(InWorld, Cmd, Ar);
+	}
+	else if (TitleFileInterfacePtr != nullptr && FParse::Command(&Cmd, TEXT("TITLEFILE"))) /* ONLINE (EOS if using EOSPlus) TITLEFILE ... */
 	{
 		bWasHandled = TitleFileInterfacePtr->HandleTitleFileExec(InWorld, Cmd, Ar);
 	}
-	else if (UserCloudInterfacePtr != nullptr && FParse::Command(&Cmd, TEXT("USERCLOUD")))
+	else if (UserCloudInterfacePtr != nullptr && FParse::Command(&Cmd, TEXT("USERCLOUD"))) /* ONLINE (EOS if using EOSPlus) USERCLOUD ... */
 	{
 		bWasHandled = UserCloudInterfacePtr->HandleUserCloudExec(InWorld, Cmd, Ar);
 	}
