PR #11164: OnlineSubsystemEOS: Fix matching purchase transaction to receipt

#11164

diff --git a/Engine/Plugins/Online/OnlineSubsystemEOS/Source/OnlineSubsystemEOS/Private/OnlineStoreEOS.cpp b/Engine/Plugins/Online/OnlineSubsystemEOS/Source/OnlineSubsystemEOS/Private/OnlineStoreEOS.cpp
--- a/Engine/Plugins/Online/OnlineSubsystemEOS/Source/OnlineSubsystemEOS/Private/OnlineStoreEOS.cpp
+++ b/Engine/Plugins/Online/OnlineSubsystemEOS/Source/OnlineSubsystemEOS/Private/OnlineStoreEOS.cpp
@@ -195,11 +195,11 @@ void FOnlineStoreEOS::Checkout(const FUniqueNetId& UserId, const FPurchaseChecko
 	FCheckoutCallback* CallbackObj = new FCheckoutCallback(FOnlineStoreEOSWeakPtr(AsShared()));
 	CallbackObj->CallbackLambda = [this, OnComplete = FOnPurchaseCheckoutComplete(Delegate)](const EOS_Ecom_CheckoutCallbackInfo* Data)
 	{
-		EOS_EResult Result = Data->ResultCode;
-		if (Result != EOS_EResult::EOS_Success)
+		const EOS_EResult CheckoutResult = Data->ResultCode;
+		if (CheckoutResult != EOS_EResult::EOS_Success)
 		{
 			UE_LOG_ONLINE(Error, TEXT("EOS_Ecom_Checkout: failed with error (%s)"), ANSI_TO_TCHAR(EOS_EResult_ToString(Data->ResultCode)));
-			if (Result == EOS_EResult::EOS_Canceled)
+			if (CheckoutResult == EOS_EResult::EOS_Canceled)
 			{
 				OnComplete.ExecuteIfBound(ONLINE_ERROR(EOnlineErrorResult::Canceled), MakeShared<FPurchaseReceipt>());
 			}
@@ -210,9 +210,42 @@ void FOnlineStoreEOS::Checkout(const FUniqueNetId& UserId, const FPurchaseChecko
 			return;
 		}
 
+		EOS_Ecom_CopyTransactionByIdOptions CopyTransactionOptions = { };
+		CopyTransactionOptions.ApiVersion = 1;
+		UE_EOS_CHECK_API_MISMATCH(EOS_ECOM_COPYTRANSACTIONBYID_API_LATEST, 1);
+		CopyTransactionOptions.LocalUserId = Data->LocalUserId;
+		CopyTransactionOptions.TransactionId = Data->TransactionId;
+		EOS_Ecom_HTransaction TransactionHandle;
+		const EOS_EResult CopyTransactionResult = EOS_Ecom_CopyTransactionById(EOSSubsystem->EcomHandle, &CopyTransactionOptions, &TransactionHandle);
+		if (CopyTransactionResult != EOS_EResult::EOS_Success)
+		{
+			OnComplete.ExecuteIfBound(ONLINE_ERROR(EOnlineErrorResult::Unknown), MakeShared<FPurchaseReceipt>());
+			return;
+		}
+
+		TSet<FString> EntitlementIds = { };
+		EOS_Ecom_Transaction_GetEntitlementsCountOptions EntitlementCountOptions = { };
+		EntitlementCountOptions.ApiVersion = 1;
+		UE_EOS_CHECK_API_MISMATCH(EOS_ECOM_TRANSACTION_GETENTITLEMENTSCOUNT_API_LATEST, 1);
+		const uint32_t EntitlementsCount = EOS_Ecom_Transaction_GetEntitlementsCount(TransactionHandle, &EntitlementCountOptions);
+		for (uint32_t i = 0; i < EntitlementsCount; i++)
+		{
+			EOS_Ecom_Transaction_CopyEntitlementByIndexOptions Options = { };
+			Options.ApiVersion = 1;
+			UE_EOS_CHECK_API_MISMATCH(EOS_ECOM_TRANSACTION_COPYENTITLEMENTBYINDEX_API_LATEST, 1);
+			EOS_Ecom_Entitlement* Entitlement;
+			const EOS_EResult CopyEntitlementResult = EOS_Ecom_Transaction_CopyEntitlementByIndex(TransactionHandle, &Options, &Entitlement);
+			if (CopyEntitlementResult == EOS_EResult::EOS_Success)
+			{
+				EntitlementIds.Add(FString(UTF8_TO_TCHAR(Entitlement->EntitlementId)));
+				EOS_Ecom_Entitlement_Release(Entitlement);
+			}
+		}
+		EOS_Ecom_Transaction_Release(TransactionHandle);
+
 		// Update the cached receipts
 		QueryReceipts(*EOSSubsystem->UserManager->GetLocalUniqueNetIdEOS(Data->LocalUserId), true,
-			FOnQueryReceiptsComplete::CreateLambda([this, PurchaseComplete = FOnPurchaseCheckoutComplete(OnComplete), TransId = FString(Data->TransactionId)](const FOnlineError& Result)
+			FOnQueryReceiptsComplete::CreateLambda([this, PurchaseComplete = FOnPurchaseCheckoutComplete(OnComplete), EntitlementIds = MoveTemp(EntitlementIds)](const FOnlineError& Result)
 		{
 			if (!Result.WasSuccessful())
 			{
@@ -225,7 +258,7 @@ void FOnlineStoreEOS::Checkout(const FUniqueNetId& UserId, const FPurchaseChecko
 			// Find the transaction in our receipts
 			for (const FPurchaseReceipt& SearchReceipt : CachedReceipts)
 			{
-				if (SearchReceipt.TransactionId == TransId)
+				if (EntitlementIds.Contains(SearchReceipt.TransactionId))
 				{
 					Receipt = MakeShared<FPurchaseReceipt>(SearchReceipt);
 					break;
