- Cleanup EOSPlus Startup/ShutdownModule to not use raw new/delete.
- Add commandlet early-out to ShutdownModule to match StartupModule
- Rename ModuleEOSPlus.cpp to EOSPlusModule.cpp to match the header.

diff --git a/Engine/Plugins/Online/OnlineSubsystemEOS/Source/OnlineSubsystemEOSPlus/Private/OnlineSubsystemEOSPlusModule.cpp b/Engine/Plugins/Online/OnlineSubsystemEOS/Source/OnlineSubsystemEOSPlus/Private/OnlineSubsystemEOSPlusModule.cpp
new file mode 100644
--- /dev/null
+++ b/Engine/Plugins/Online/OnlineSubsystemEOS/Source/OnlineSubsystemEOSPlus/Private/OnlineSubsystemEOSPlusModule.cpp
@@ -0,0 +1,92 @@
+// Copyright Epic Games, Inc. All Rights Reserved.
+
+#include "OnlineSubsystemEOSPlusModule.h"
+#include "OnlineSubsystemModule.h"
+#include "OnlineSubsystemNames.h"
+#include "OnlineSubsystemEOSPlus.h"
+
+#include "Features/IModularFeature.h"
+#include "Features/IModularFeatures.h"
+#include "Misc/ConfigCacheIni.h"
+
+#define LOCTEXT_NAMESPACE "EOSPlus"
+
+IMPLEMENT_MODULE(FOnlineSubsystemEOSPlusModule, OnlineSubsystemEOSPlus);
+
+/**
+ * Class responsible for creating instance(s) of the subsystem
+ */
+class FOnlineFactoryEOSPlus :
+	public IOnlineFactory
+{
+public:
+
+	FOnlineFactoryEOSPlus() {}
+	virtual ~FOnlineFactoryEOSPlus() {}
+
+	virtual IOnlineSubsystemPtr CreateSubsystem(FName InstanceName)
+	{
+		FOnlineSubsystemEOSPlusPtr OnlineSub = MakeShared<FOnlineSubsystemEOSPlus, ESPMode::ThreadSafe>(InstanceName);
+		if (!OnlineSub->Init())
+		{
+			UE_LOG_ONLINE(Warning, TEXT("EOSPlus failed to initialize!"));
+			OnlineSub->Shutdown();
+			OnlineSub = nullptr;
+		}
+
+		return OnlineSub;
+	}
+};
+
+void FOnlineSubsystemEOSPlusModule::StartupModule()
+{
+	if (IsRunningCommandlet())
+	{
+		return;
+	}
+
+	// Force loading of the EOS OSS module in case the plugin manager failed to get the dependencies right
+	FModuleManager::LoadModuleChecked<FOnlineSubsystemModule>(TEXT("OnlineSubsystemEOS"));
+	// Make sure the base OSS has a chance to register its OSS factory
+	LoadBaseOSS();
+
+	PlusFactory = MakeUnique<FOnlineFactoryEOSPlus>();
+
+	// Create and register our singleton factory with the main online subsystem for easy access
+	FOnlineSubsystemModule& OSS = FModuleManager::GetModuleChecked<FOnlineSubsystemModule>("OnlineSubsystem");
+	OSS.RegisterPlatformService(EOSPLUS_SUBSYSTEM, PlusFactory.Get());
+}
+
+void FOnlineSubsystemEOSPlusModule::ShutdownModule()
+{
+	if (IsRunningCommandlet())
+	{
+		return;
+	}
+
+	FOnlineSubsystemModule& OSS = FModuleManager::GetModuleChecked<FOnlineSubsystemModule>("OnlineSubsystem");
+	OSS.UnregisterPlatformService(EOSPLUS_SUBSYSTEM);
+
+	PlusFactory.Reset();
+}
+
+void FOnlineSubsystemEOSPlusModule::LoadBaseOSS()
+{
+	FString BaseOSSName;
+	GConfig->GetString(TEXT("[OnlineSubsystemEOSPlus]"), TEXT("BaseOSSName"), BaseOSSName, GEngineIni);
+	if (BaseOSSName.IsEmpty())
+	{
+		// Load the native platform OSS name
+		GConfig->GetString(TEXT("OnlineSubsystem"), TEXT("NativePlatformService"), BaseOSSName, GEngineIni);
+	}
+	if (BaseOSSName.IsEmpty())
+	{
+		UE_LOG_ONLINE(Verbose, TEXT("EOSPlus failed to find the native OSS!"));
+		return;
+	}
+	// Load the module for the base OSS
+	FString ModuleName(TEXT("OnlineSubsystem") + BaseOSSName);
+	FModuleManager::LoadModuleChecked<FOnlineSubsystemModule>(*ModuleName);
+}
+
+#undef LOCTEXT_NAMESPACE
diff --git a/Engine/Plugins/Online/OnlineSubsystemEOS/Source/OnlineSubsystemEOSPlus/Public/OnlineSubsystemEOSPlusModule.h b/Engine/Plugins/Online/OnlineSubsystemEOS/Source/OnlineSubsystemEOSPlus/Public/OnlineSubsystemEOSPlusModule.h
--- a/Engine/Plugins/Online/OnlineSubsystemEOS/Source/OnlineSubsystemEOSPlus/Public/OnlineSubsystemEOSPlusModule.h
+++ b/Engine/Plugins/Online/OnlineSubsystemEOS/Source/OnlineSubsystemEOSPlus/Public/OnlineSubsystemEOSPlusModule.h
@@ -4,6 +4,7 @@
 
 #include "CoreMinimal.h"
 #include "Modules/ModuleInterface.h"
+#include "Templates/UniquePtr.h"
 
 /**
  * OSS that aggregates EOS plus the platform OSS
@@ -11,13 +12,7 @@
 class FOnlineSubsystemEOSPlusModule : public IModuleInterface
 {
 public:
-
-	FOnlineSubsystemEOSPlusModule() :
-		PlusFactory(nullptr)
-	{}
-
-	virtual ~FOnlineSubsystemEOSPlusModule() {}
-
+	virtual ~FOnlineSubsystemEOSPlusModule() = default;
 
 // IModuleInterface
 	virtual void StartupModule() override;
@@ -37,5 +32,5 @@ private:
 	void LoadBaseOSS();
 
 	/** Class responsible for creating instance(s) of the subsystem */
-	class FOnlineFactoryEOSPlus* PlusFactory;
+	TUniquePtr<class FOnlineFactoryEOSPlus> PlusFactory;
 };
