Minor fix - Overlay was incorrectly disabled when not playing in the editor.

diff --git a/Engine/Plugins/Online/OnlineSubsystemEOS/Source/OnlineSubsystemEOS/Private/OnlineSubsystemEOS.cpp b/Engine/Plugins/Online/OnlineSubsystemEOS/Source/OnlineSubsystemEOS/Private/OnlineSubsystemEOS.cpp
--- a/Engine/Plugins/Online/OnlineSubsystemEOS/Source/OnlineSubsystemEOS/Private/OnlineSubsystemEOS.cpp
+++ b/Engine/Plugins/Online/OnlineSubsystemEOS/Source/OnlineSubsystemEOS/Private/OnlineSubsystemEOS.cpp
@@ -211,7 +211,7 @@ bool FOnlineSubsystemEOS::PlatformCreate()
 		OverlayFlags |= EOS_PF_DISABLE_SOCIAL_OVERLAY;
 	}
 #if WITH_EDITOR
-	if (!EOSSettings.bEnableEditorOverlay)
+	if (GIsEditor && EOSSettings.bEnableEditorOverlay)
 	{
 		OverlayFlags |= EOS_PF_LOADING_IN_EDITOR;
 	}
