Fix some dependencies
#rnx

diff --git a/Engine/Plugins/Online/OnlineSubsystemEOS/Source/OnlineSubsystemEOS/OnlineSubsystemEOS.Build.cs b/Engine/Plugins/Online/OnlineSubsystemEOS/Source/OnlineSubsystemEOS/OnlineSubsystemEOS.Build.cs
--- a/Engine/Plugins/Online/OnlineSubsystemEOS/Source/OnlineSubsystemEOS/OnlineSubsystemEOS.Build.cs
+++ b/Engine/Plugins/Online/OnlineSubsystemEOS/Source/OnlineSubsystemEOS/OnlineSubsystemEOS.Build.cs
@@ -16,6 +16,13 @@ public class OnlineSubsystemEOS : ModuleRules
 			}
 		);
 
+		PrivateIncludePathModuleNames.AddRange(
+			new string[]
+			{
+				"VoiceChat",
+			}
+		);
+
 		PrivateDependencyModuleNames.AddRange(
 			new string[] {
 				"Core",
@@ -29,7 +36,6 @@ public class OnlineSubsystemEOS : ModuleRules
 				"OnlineSubsystemUtils",
 				"Sockets",
 				"SocketSubsystemEOS",
-				"VoiceChat",
 				"NetCore"
 			}
 		);
