Changed how OSS EOS Sessions responds to a EOS_ELobbyMemberStatus::EOS_LMS_CLOSED notification, destroying the related lobby instead of just removing the lobby member.

diff --git a/Engine/Plugins/Online/OnlineSubsystemEOS/Source/OnlineSubsystemEOS/Private/OnlineSessionEOS.cpp b/Engine/Plugins/Online/OnlineSubsystemEOS/Source/OnlineSubsystemEOS/Private/OnlineSessionEOS.cpp
--- a/Engine/Plugins/Online/OnlineSubsystemEOS/Source/OnlineSubsystemEOS/Private/OnlineSessionEOS.cpp
+++ b/Engine/Plugins/Online/OnlineSubsystemEOS/Source/OnlineSubsystemEOS/Private/OnlineSessionEOS.cpp
@@ -1016,20 +1016,8 @@ void FOnlineSessionEOS::OnMemberStatusReceived(const EOS_LobbyId& LobbyId, const
 			break;
 		case EOS_ELobbyMemberStatus::EOS_LMS_CLOSED:
 			{
-				EOSSubsystem->UserManager->ResolveUniqueNetId(EOSSubsystem->UserManager->GetDefaultLocalUser(), TargetUserId, [this, LobbyNetId](FUniqueNetIdEOSRef ResolvedUniqueNetId)
-					{
-						FNamedOnlineSession* Session = GetNamedSessionFromLobbyId(*LobbyNetId);
-						if (Session)
-						{
-							RemoveOnlineSessionMember(Session->SessionName, ResolvedUniqueNetId);
-
-							TriggerOnSessionParticipantLeftDelegates(Session->SessionName, *ResolvedUniqueNetId, EOnSessionParticipantLeftReason::Closed);
-						}
-						else
-						{
-							UE_LOG_ONLINE_SESSION(VeryVerbose, TEXT("[FOnlineSessionEOS::OnMemberStatusReceived] EOS_LMS_CLOSED: Unable to retrieve session with LobbyId %s"), *LobbyNetId->ToString());
-						}
-					});
+				const int32 DefaultLocalUser = EOSSubsystem->UserManager->GetDefaultLocalUser();
+				DestroyLobbySession(DefaultLocalUser, Session, FOnDestroySessionCompleteDelegate::CreateLambda([](FName SessionName, bool bWasSuccessful) {}));
 			}
 			break;
 		}
