Added instance name to OSS creation calls in EOSPlus

diff --git a/Engine/Plugins/Online/OnlineSubsystemEOS/Source/OnlineSubsystemEOSPlus/Private/OnlineSubsystemEOSPlus.cpp b/Engine/Plugins/Online/OnlineSubsystemEOS/Source/OnlineSubsystemEOSPlus/Private/OnlineSubsystemEOSPlus.cpp
--- a/Engine/Plugins/Online/OnlineSubsystemEOS/Source/OnlineSubsystemEOSPlus/Private/OnlineSubsystemEOSPlus.cpp
+++ b/Engine/Plugins/Online/OnlineSubsystemEOS/Source/OnlineSubsystemEOSPlus/Private/OnlineSubsystemEOSPlus.cpp
@@ -49,7 +49,8 @@ bool FOnlineSubsystemEOSPlus::Init()
 		return false;
 	}
 
-	BaseOSS = IOnlineSubsystem::Get(FName(*BaseOSSName));
+	const FString BaseSubsystemFullName = BaseOSSName + TEXT(":") + InstanceName.ToString();
+	BaseOSS = IOnlineSubsystem::Get(FName(*BaseSubsystemFullName));
 	if (BaseOSS == nullptr)
 	{
 		UE_LOG_ONLINE(Error, TEXT("FOnlineSubsystemEOSPlus::Init() failed to get the platform OSS"));
@@ -63,7 +64,8 @@ bool FOnlineSubsystemEOSPlus::Init()
 		return false;
 	}
 
-	EosOSS = IOnlineSubsystem::Get(EOS_SUBSYSTEM);
+	const FString EOSSubsystemFullName = EOS_SUBSYSTEM.ToString() + TEXT(":") + InstanceName.ToString();
+	EosOSS = IOnlineSubsystem::Get(FName(*EOSSubsystemFullName));
 	if (EosOSS == nullptr)
 	{
 		UE_LOG_ONLINE(Error, TEXT("FOnlineSubsystemEOSPlus::Init() failed to get the EOS OSS"));
