PR #11419: Added missing delegate trigger for presence updates in UserManagerEOS

diff --git a/Engine/Plugins/Online/OnlineSubsystemEOS/Source/OnlineSubsystemEOS/Private/UserManagerEOS.cpp b/Engine/Plugins/Online/OnlineSubsystemEOS/Source/OnlineSubsystemEOS/Private/UserManagerEOS.cpp
--- a/Engine/Plugins/Online/OnlineSubsystemEOS/Source/OnlineSubsystemEOS/Private/UserManagerEOS.cpp
+++ b/Engine/Plugins/Online/OnlineSubsystemEOS/Source/OnlineSubsystemEOS/Private/UserManagerEOS.cpp
@@ -3032,6 +3032,8 @@ void FUserManagerEOS::UpdatePresence(int32 LocalUserNum, EOS_EpicAccountId Accou
 
 		// Copy the presence if this is a friend that was updated, so that their data is in sync
 		UpdateFriendPresence(NetId, PresenceRef);
+		
+		TriggerOnPresenceReceivedDelegates(*NetId, PresenceRef);
 
 		EOS_Presence_Info_Release(PresenceInfo);
 	}
