Added exec handling logic to OSSEOS Sessions
Removed outdated comment from FindFriendSession

diff --git a/Engine/Plugins/Online/OnlineSubsystemEOS/Source/OnlineSubsystemEOS/Private/OnlineSessionEOS.cpp b/Engine/Plugins/Online/OnlineSubsystemEOS/Source/OnlineSubsystemEOS/Private/OnlineSessionEOS.cpp
--- a/Engine/Plugins/Online/OnlineSubsystemEOS/Source/OnlineSubsystemEOS/Private/OnlineSessionEOS.cpp
+++ b/Engine/Plugins/Online/OnlineSubsystemEOS/Source/OnlineSubsystemEOS/Private/OnlineSessionEOS.cpp
@@ -2858,7 +2858,6 @@ bool FOnlineSessionEOS::FindFriendSession(int32 LocalUserNum, const FUniqueNetId
 		UE_EOS_CHECK_API_MISMATCH(EOS_LOBBYSEARCH_SETTARGETUSERID_API_LATEST, 1);
 		SetTargetUserIdOptions.TargetUserId = FriendEOSId.GetProductUserId();
 
-		// TODO: Using this as a search parameter only works if we use the owner's id (search for lobbies we're already in). Pending API fix so it works with other users too.
 		EOS_LobbySearch_SetTargetUserId(LobbySearchHandle, &SetTargetUserIdOptions);
 
 		// Then perform the search
@@ -4940,4 +4939,39 @@ void FOnlineSessionEOS::CopyLobbyMemberAttributes(const FLobbyDetailsEOS& LobbyD
 	}
 }
 
+bool FOnlineSessionEOS::HandleSessionExec(UWorld* InWorld, const TCHAR* Cmd, FOutputDevice& Ar)
+{
+#if !UE_BUILD_SHIPPING
+	bool bWasHandled = true;
+
+	if (FParse::Command(&Cmd, TEXT("FindFriendSession"))) /* ONLINE (EOS if using EOSPlus) SESSION FindFriendSession LocalUserNum=0 FriendListName=default FriendIndex=0 */
+	{
+		int LocalUserNum = 0;
+		FParse::Value(Cmd, TEXT("LocalUserNum="), LocalUserNum);
+
+		FString FriendListName;
+		FParse::Value(Cmd, TEXT("FriendListName="), FriendListName);
+
+		int FriendIndex = 0;
+		FParse::Value(Cmd, TEXT("FriendIndex="), FriendIndex);
+		
+		TArray<TSharedRef<FOnlineFriend>> FriendList;
+		EOSSubsystem->UserManager->GetFriendsList(LocalUserNum, FriendListName, FriendList);
+		if (FriendList.Num() > FriendIndex)
+		{
+			const TSharedRef<FOnlineFriend>& Friend = FriendList[FriendIndex];
+			FindFriendSession(LocalUserNum, *Friend->GetUserId());
+		}
+	}
+	else
+	{
+		bWasHandled = false;
+	}
+
+	return bWasHandled;
+#else
+	return false;
+#endif // !UE_BUILD_SHIPPING
+}
+
 #endif // WITH_EOS_SDK
diff --git a/Engine/Plugins/Online/OnlineSubsystemEOS/Source/OnlineSubsystemEOS/Private/OnlineSessionEOS.h b/Engine/Plugins/Online/OnlineSubsystemEOS/Source/OnlineSubsystemEOS/Private/OnlineSessionEOS.h
--- a/Engine/Plugins/Online/OnlineSubsystemEOS/Source/OnlineSubsystemEOS/Private/OnlineSessionEOS.h
+++ b/Engine/Plugins/Online/OnlineSubsystemEOS/Source/OnlineSubsystemEOS/Private/OnlineSessionEOS.h
@@ -193,6 +193,8 @@ public:
 
 	void Init();
 
+	bool HandleSessionExec(UWorld* InWorld, const TCHAR* Cmd, FOutputDevice& Ar);
+
 private:
 	// EOS Lobbies
 
diff --git a/Engine/Plugins/Online/OnlineSubsystemEOS/Source/OnlineSubsystemEOS/Private/OnlineSubsystemEOS.cpp b/Engine/Plugins/Online/OnlineSubsystemEOS/Source/OnlineSubsystemEOS/Private/OnlineSubsystemEOS.cpp
--- a/Engine/Plugins/Online/OnlineSubsystemEOS/Source/OnlineSubsystemEOS/Private/OnlineSubsystemEOS.cpp
+++ b/Engine/Plugins/Online/OnlineSubsystemEOS/Source/OnlineSubsystemEOS/Private/OnlineSubsystemEOS.cpp
@@ -510,6 +510,10 @@ bool FOnlineSubsystemEOS::Exec(UWorld* InWorld, const TCHAR* Cmd, FOutputDevice&
 	{
 		bWasHandled = LeaderboardsInterfacePtr->HandleLeaderboardsExec(InWorld, Cmd, Ar);
 	}
+	else if (SessionInterfacePtr != nullptr && FParse::Command(&Cmd, TEXT("SESSION"))) /* ONLINE (EOS if using EOSPlus) SESSION ... */
+	{
+		bWasHandled = SessionInterfacePtr->HandleSessionExec(InWorld, Cmd, Ar);
+	}
 	else if (TitleFileInterfacePtr != nullptr && FParse::Command(&Cmd, TEXT("TITLEFILE"))) /* ONLINE (EOS if using EOSPlus) TITLEFILE ... */
 	{
 		bWasHandled = TitleFileInterfacePtr->HandleTitleFileExec(InWorld, Cmd, Ar);
