- Fix session ports not being correctly replicated in OSSEOS, so clients always try to join on the default port.
- Note there is an issue where EOS only supports int64, so if you set an attribute as int32, it will get promoted to int64 on its roundtrip to the backend, and it will no longer be accessible via FOnlineSessionSettings::Get which a) does not have an int64 override and b) how would you know which one to call? We work around this by finding the underlying attribute directly and handling both cases. Future work to resolve this issue so licensees don't have to do the same thing on all their integer attribs.

diff --git a/Engine/Plugins/Online/OnlineSubsystemEOS/Source/OnlineSubsystemEOS/Private/OnlineSessionEOS.cpp b/Engine/Plugins/Online/OnlineSubsystemEOS/Source/OnlineSubsystemEOS/Private/OnlineSessionEOS.cpp
--- a/Engine/Plugins/Online/OnlineSubsystemEOS/Source/OnlineSubsystemEOS/Private/OnlineSessionEOS.cpp
+++ b/Engine/Plugins/Online/OnlineSubsystemEOS/Source/OnlineSubsystemEOS/Private/OnlineSessionEOS.cpp
@@ -360,52 +360,27 @@ struct FLobbyAttributeOptions :
 	}
 };
 
-FOnlineSessionInfoEOS::FOnlineSessionInfoEOS()
-	: HostAddr(nullptr)
-	, SessionId(FUniqueNetIdString::EmptyId())
-	, SessionHandle(nullptr)
-	, LobbyHandle(nullptr)
-	, bIsFromClone(false)
+FOnlineSessionInfoEOS FOnlineSessionInfoEOS::Create(FUniqueNetIdStringRef UniqueNetId)
 {
+	FOnlineSessionInfoEOS Result;
+	Result.SessionId = UniqueNetId;
+	return Result;
 }
 
-FOnlineSessionInfoEOS::FOnlineSessionInfoEOS(const FString& InHostIp, FUniqueNetIdStringRef UniqueNetId, const TSharedPtr<FSessionDetailsEOS>& InSessionHandle, const TSharedPtr<FLobbyDetailsEOS>& InLobbyHandle)
-	: FOnlineSessionInfo()
-	, SessionId(UniqueNetId)
-	, SessionHandle(InSessionHandle)
-	, LobbyHandle(InLobbyHandle)
-	, bIsFromClone(false)
-{
-	if (InHostIp.StartsWith(EOS_CONNECTION_URL_PREFIX, ESearchCase::IgnoreCase))
-	{
-		HostAddr = ISocketSubsystem::Get(EOS_SOCKETSUBSYSTEM)->GetAddressFromString(InHostIp);
-		EOSAddress = InHostIp;
-	}
-	else
-	{
-		HostAddr = ISocketSubsystem::Get(PLATFORM_SOCKETSUBSYSTEM)->GetAddressFromString(InHostIp);
-		HostAddr->SetPort(FURL::UrlConfig.DefaultPort);
-	}
-}
-
-FOnlineSessionInfoEOS::~FOnlineSessionInfoEOS()
-{
-
-}
-
-FOnlineSessionInfoEOS FOnlineSessionInfoEOS::Create(const FString& InHostIp, FUniqueNetIdStringRef UniqueNetId)
-{
-	return FOnlineSessionInfoEOS(InHostIp, UniqueNetId, nullptr, nullptr);
-}
-
-FOnlineSessionInfoEOS FOnlineSessionInfoEOS::Create(const FString& InHostIp, FUniqueNetIdStringRef UniqueNetId, const TSharedPtr<FSessionDetailsEOS>& InSessionHandle)
+FOnlineSessionInfoEOS FOnlineSessionInfoEOS::Create(FUniqueNetIdStringRef UniqueNetId, const TSharedPtr<FSessionDetailsEOS>& SessionHandle)
 {
-	return FOnlineSessionInfoEOS(InHostIp, UniqueNetId, InSessionHandle, nullptr);
+	FOnlineSessionInfoEOS Result;
+	Result.SessionId = UniqueNetId;
+	Result.SessionHandle = SessionHandle;
+	return Result;
 }
 
-FOnlineSessionInfoEOS FOnlineSessionInfoEOS::Create(const FString& InHostIp, FUniqueNetIdStringRef UniqueNetId, const TSharedPtr<FLobbyDetailsEOS>& InLobbyHandle)
+FOnlineSessionInfoEOS FOnlineSessionInfoEOS::Create(FUniqueNetIdStringRef UniqueNetId, const TSharedPtr<FLobbyDetailsEOS>& LobbyHandle)
 {
-	return FOnlineSessionInfoEOS(InHostIp, UniqueNetId, nullptr, InLobbyHandle);
+	FOnlineSessionInfoEOS Result;
+	Result.SessionId = UniqueNetId;
+	Result.LobbyHandle = LobbyHandle;
+	return Result;
 }
 
 void FOnlineSessionInfoEOS::InitLAN(FOnlineSubsystemEOS* Subsystem)
@@ -696,7 +671,7 @@ void FOnlineSessionEOS::OnSessionInviteReceived(const EOS_Sessions_SessionInvite
 			}
 			else
 			{
-				UE_LOG_ONLINE_SESSION(Warning, TEXT("EOS_Sessions_CopySessionHandleByInviteId not successful. Finished with EOS_EResult %s"), ANSI_TO_TCHAR(EOS_EResult_ToString(Result)));
+				UE_LOG_ONLINE_SESSION(Warning, TEXT("EOS_Sessions_CopySessionHandleByInviteId not successful. Finished with EOS_EResult %s"), *LexToString(Result));
 				TriggerOnSessionInviteReceivedDelegates(*NetId, *FromNetId, EOSSubsystem->GetAppId(), FOnlineSessionSearchResult());
 			}
 		});
@@ -725,7 +700,7 @@ void FOnlineSessionEOS::OnSessionInviteAccepted(const EOS_Sessions_SessionInvite
 			}
 			else
 			{
-				UE_LOG_ONLINE_SESSION(Warning, TEXT("EOS_Sessions_CopySessionHandleByInviteId not successful. Finished with EOS_EResult %s"), ANSI_TO_TCHAR(EOS_EResult_ToString(Result)));
+				UE_LOG_ONLINE_SESSION(Warning, TEXT("EOS_Sessions_CopySessionHandleByInviteId not successful. Finished with EOS_EResult %s"), *LexToString(Result));
 				TriggerOnSessionUserInviteAcceptedDelegates(false, LocalUserNum, ResolvedUniqueNetId, FOnlineSessionSearchResult());
 			}
 		});
@@ -753,7 +728,7 @@ void FOnlineSessionEOS::OnJoinSessionAccepted(const EOS_Sessions_JoinSessionAcce
 			}
 			else
 			{
-				UE_LOG_ONLINE_SESSION(Warning, TEXT("EOS_Sessions_CopySessionHandleByUiEventId not successful. Finished with EOS_EResult %s"), ANSI_TO_TCHAR(EOS_EResult_ToString(Result)));
+				UE_LOG_ONLINE_SESSION(Warning, TEXT("EOS_Sessions_CopySessionHandleByUiEventId not successful. Finished with EOS_EResult %s"), *LexToString(Result));
 				TriggerOnSessionUserInviteAcceptedDelegates(false, LocalUserNum, ResolvedUniqueNetId, FOnlineSessionSearchResult());
 			}
 		});
@@ -888,12 +863,12 @@ void FOnlineSessionEOS::OnLobbyUpdateReceived(const EOS_LobbyId& LobbyId)
 			}
 			else
 			{
-				UE_LOG_ONLINE_SESSION(Warning, TEXT("[FOnlineSessionEOS::OnLobbyUpdateReceived] EOS_LobbyDetails_CopyInfo not successful. Finished with EOS_EResult %s"), ANSI_TO_TCHAR(EOS_EResult_ToString(CopyInfoResult)));
+				UE_LOG_ONLINE_SESSION(Warning, TEXT("[FOnlineSessionEOS::OnLobbyUpdateReceived] EOS_LobbyDetails_CopyInfo not successful. Finished with EOS_EResult %s"), *LexToString(CopyInfoResult));
 			}
 		}
 		else
 		{
-			UE_LOG_ONLINE_SESSION(Warning, TEXT("[FOnlineSessionEOS::OnLobbyUpdateReceived] EOS_Lobby_CopyLobbyDetailsHandle not successful. Finished with EOS_EResult %s"), ANSI_TO_TCHAR(EOS_EResult_ToString(CopyLobbyDetailsResult)));
+			UE_LOG_ONLINE_SESSION(Warning, TEXT("[FOnlineSessionEOS::OnLobbyUpdateReceived] EOS_Lobby_CopyLobbyDetailsHandle not successful. Finished with EOS_EResult %s"), *LexToString(CopyLobbyDetailsResult));
 		}
 	}
 	else
@@ -1071,7 +1046,7 @@ void FOnlineSessionEOS::OnLobbyInviteReceived(const EOS_Lobby_LobbyInviteReceive
 			}
 			else
 			{
-				UE_LOG_ONLINE_SESSION(Warning, TEXT("EOS_Lobby_CopyLobbyDetailsHandleByInviteId not successful. Finished with EOS_EResult %s"), ANSI_TO_TCHAR(EOS_EResult_ToString(Result)));
+				UE_LOG_ONLINE_SESSION(Warning, TEXT("EOS_Lobby_CopyLobbyDetailsHandleByInviteId not successful. Finished with EOS_EResult %s"), *LexToString(Result));
 				TriggerOnSessionInviteReceivedDelegates(*NetId, *FromNetId, EOSSubsystem->GetAppId(), FOnlineSessionSearchResult());
 			}
 		});
@@ -1105,7 +1080,7 @@ void FOnlineSessionEOS::OnLobbyInviteAccepted(const EOS_Lobby_LobbyInviteAccepte
 			}
 			else
 			{
-				UE_LOG_ONLINE_SESSION(Warning, TEXT("[FOnlineSessionEOS::OnLobbyInviteAccepted] EOS_Lobby_CopyLobbyDetailsHandleByInviteId failed with EOS result code (%s)"), ANSI_TO_TCHAR(EOS_EResult_ToString(Result)));
+				UE_LOG_ONLINE_SESSION(Warning, TEXT("[FOnlineSessionEOS::OnLobbyInviteAccepted] EOS_Lobby_CopyLobbyDetailsHandleByInviteId failed with EOS result code (%s)"), *LexToString(Result));
 				TriggerOnSessionUserInviteAcceptedDelegates(false, LocalUserNum, ResolvedUniqueNetId, FOnlineSessionSearchResult());
 			}
 		});
@@ -1137,7 +1112,7 @@ void FOnlineSessionEOS::OnJoinLobbyAccepted(const EOS_Lobby_JoinLobbyAcceptedCal
 			}
 			else
 			{
-				UE_LOG_ONLINE_SESSION(Warning, TEXT("[FOnlineSessionEOS::OnJoinLobbyAccepted] EOS_Lobby_CopyLobbyDetailsHandleByUiEventId failed with EOS result code (%s)"), ANSI_TO_TCHAR(EOS_EResult_ToString(Result)));
+				UE_LOG_ONLINE_SESSION(Warning, TEXT("[FOnlineSessionEOS::OnJoinLobbyAccepted] EOS_Lobby_CopyLobbyDetailsHandleByUiEventId failed with EOS result code (%s)"), *LexToString(Result));
 				TriggerOnSessionUserInviteAcceptedDelegates(false, LocalUserNum, ResolvedUniqueNetId, FOnlineSessionSearchResult());
 			}
 		});
@@ -1437,7 +1412,7 @@ void FOnlineSessionEOS::BeginSessionAnalytics(FNamedOnlineSession* Session)
 		EOS_EResult Result = EOS_Metrics_BeginPlayerSession(EOSSubsystem->MetricsHandle, &Options);
 		if (Result != EOS_EResult::EOS_Success)
 		{
-			UE_LOG_ONLINE_SESSION(Error, TEXT("EOS_Metrics_BeginPlayerSession() returned EOS result code (%s)"), ANSI_TO_TCHAR(EOS_EResult_ToString(Result)));
+			UE_LOG_ONLINE_SESSION(Error, TEXT("EOS_Metrics_BeginPlayerSession() returned EOS result code (%s)"), *LexToString(Result));
 		}
 	}
 	else
@@ -1471,6 +1446,50 @@ struct FSessionCreateOptions :
 	}
 };
 
+TSharedPtr<class FInternetAddr> GetHostAddress(const FString& EosHostAddress, const FOnlineSessionSettings& Settings)
+{
+	TSharedPtr<class FInternetAddr> Result;
+
+	if (EosHostAddress.StartsWith(EOS_CONNECTION_URL_PREFIX, ESearchCase::IgnoreCase))
+	{
+		Result = ISocketSubsystem::Get(EOS_SOCKETSUBSYSTEM)->GetAddressFromString(EosHostAddress);
+	}
+	else
+	{
+		if (Result = ISocketSubsystem::Get(PLATFORM_SOCKETSUBSYSTEM)->GetAddressFromString(EosHostAddress))
+		{
+			// Read port from settings, or fall back on DefaultPort.
+			int32 Port = FURL::UrlConfig.DefaultPort;
+
+			// TODO Settings.Get _should_ work, BUT even if you set it as int32, the round trip to EOS sets it to int64, and there is no FOnlineSessionSettings::Get overload for int64...
+			// so tl;dr int attribute types are not stable in osseos and you have to check for all possible variations...
+			if (const FOnlineSessionSetting* Setting = Settings.Settings.Find(SESSION_ATTR_SERVERPORT))
+			{
+				switch (Setting->Data.GetType())
+				{
+					case EOnlineKeyValuePairDataType::Int32:
+					{
+						Setting->Data.GetValue(Port);
+						break;
+					}
+					case EOnlineKeyValuePairDataType::Int64:
+					{
+						int64 Port64;
+						Setting->Data.GetValue(Port64);
+						Port = Port64;
+						break;
+					}
+					default: checkNoEntry();
+				}
+			}
+
+			Result->SetPort(Port);
+		}
+	}
+
+	return Result;
+}
+
 uint32 FOnlineSessionEOS::CreateEOSSession(int32 HostingPlayerNum, FNamedOnlineSession* Session)
 {
 	check(Session != nullptr);
@@ -1498,7 +1517,7 @@ uint32 FOnlineSessionEOS::CreateEOSSession(int32 HostingPlayerNum, FNamedOnlineS
 	EOS_EResult ResultCode = EOS_Sessions_CreateSessionModification(EOSSubsystem->SessionsHandle, &Options, &SessionModHandle);
 	if (ResultCode != EOS_EResult::EOS_Success)
 	{
-		UE_LOG_ONLINE_SESSION(Error, TEXT("EOS_Sessions_CreateSessionModification() failed with EOS result code (%s)"), ANSI_TO_TCHAR(EOS_EResult_ToString(ResultCode)));
+		UE_LOG_ONLINE_SESSION(Error, TEXT("%hs EOS_Sessions_CreateSessionModification failed Result=%s"), __FUNCTION__, *LexToString(ResultCode));
 		return ONLINE_FAIL;
 	}
 
@@ -1506,12 +1525,12 @@ uint32 FOnlineSessionEOS::CreateEOSSession(int32 HostingPlayerNum, FNamedOnlineS
 	Session->bHosting = true;
 
 	FString HostAddr;
-	// If we are using local IPs in a dedicated server, or if we are using p2p sockets, then we need to add a custom URL for connecting
+
 	if (!IsRunningDedicatedServer() && bIsUsingP2PSockets)
 	{
 		// Because some platforms remap ports, we will use the ID of the name of the net driver to be our port instead
 		FName NetDriverName = GetDefault<UNetDriverEOS>()->NetDriverName;
-		FInternetAddrEOS TempAddr(LexToString(Options.LocalUserId), NetDriverName.ToString(), GetTypeHash(NetDriverName.ToString()));
+		FInternetAddrEOS TempAddr(Options.LocalUserId, NetDriverName.ToString(), GetTypeHash(NetDriverName.ToString()));
 		HostAddr = TempAddr.ToString(true);
 	}
 	else
@@ -1522,6 +1541,13 @@ uint32 FOnlineSessionEOS::CreateEOSSession(int32 HostingPlayerNum, FNamedOnlineS
 			bool bCanBindAll;
 			HostAddr = ISocketSubsystem::Get(PLATFORM_SOCKETSUBSYSTEM)->GetLocalHostAddr(*GLog, bCanBindAll)->ToString(false);
 		}
+		else
+		{
+			// Leave HostAddr blank, we will skip calling SetHostAddress, and EOS will set it to our public IP.
+		}
+
+		// We're using IP so need to share which port.
+		Session->SessionSettings.Set(SESSION_ATTR_SERVERPORT, FURL::UrlConfig.DefaultPort, EOnlineDataAdvertisementType::ViaOnlineService);
 	}
 
 	if (!HostAddr.IsEmpty())
@@ -1536,24 +1562,26 @@ uint32 FOnlineSessionEOS::CreateEOSSession(int32 HostingPlayerNum, FNamedOnlineS
 		// Expect URLs to look like "EOS:PUID:SocketName:Channel" and channel can be optional
 		HostOptions.HostAddress = HostAddrAnsi;
 		EOS_EResult HostResult = EOS_SessionModification_SetHostAddress(SessionModHandle, &HostOptions);
-		UE_LOG_ONLINE_SESSION(Verbose, TEXT("[FOnlineSessionEOS::CreateEOSSession] EOS_SessionModification_SetHostAddress(%s) returned (%s)"), *HostAddr, ANSI_TO_TCHAR(EOS_EResult_ToString(HostResult)));
+		UE_LOG_ONLINE_SESSION(Verbose, TEXT("%hs EOS_SessionModification_SetHostAddress(%s) returned (%s)"), __FUNCTION__, *HostAddr, *LexToString(HostResult));
 	}
 	else
 	{
 		// We'll set HostAddr locally, but it'll be ignored on the EOS API side
 		HostAddr = TEXT("127.0.0.1");
 
-		UE_LOG_ONLINE_SESSION(Verbose, TEXT("[FOnlineSessionEOS::CreateEOSSession] The server's public IP Address will be set as the Session's HostAddress."));
+		UE_LOG_ONLINE_SESSION(Verbose, TEXT("%hs The server's public IP Address will be set as the Session's HostAddress."), __FUNCTION__);
 	}
 
-	UE_LOG_ONLINE_SESSION(Verbose, TEXT("[FOnlineSessionEOS::CreateEOSSession] The HostAddress used for this session will be %s"), *HostAddr);
 
-	Session->SessionInfo = MakeShared<FOnlineSessionInfoEOS>(FOnlineSessionInfoEOS::Create(HostAddr, FUniqueNetIdEOSSession::Create(FString())));
+	// TODO why can't this use FUniqueNetIdEOSSession::EmptyId()?
+	TSharedPtr<FOnlineSessionInfoEOS> SessionInfo = MakeShared<FOnlineSessionInfoEOS>(FOnlineSessionInfoEOS::Create(FUniqueNetIdEOSSession::Create(FString())));
+	SessionInfo->HostAddr = GetHostAddress(HostAddr, Session->SessionSettings);
+	Session->SessionInfo = SessionInfo;
 
-	FName SessionName = Session->SessionName;
+	UE_LOG_ONLINE_SESSION(Verbose, TEXT("%hs The HostAddress used for this session will be %s"), __FUNCTION__, *SessionInfo->HostAddr->ToString(true));
 
 	FUpdateSessionCallback* CallbackObj = new FUpdateSessionCallback(FOnlineSessionEOSWeakPtr(AsShared()));
-	CallbackObj->CallbackLambda = [this, SessionName](const EOS_Sessions_UpdateSessionCallbackInfo* Data)
+	CallbackObj->CallbackLambda = [this, SessionName = Session->SessionName](const EOS_Sessions_UpdateSessionCallbackInfo* Data)
 	{
 		bool bWasSuccessful = false;
 
@@ -1576,7 +1604,7 @@ uint32 FOnlineSessionEOS::CreateEOSSession(int32 HostingPlayerNum, FNamedOnlineS
 			}
 			else
 			{
-				UE_LOG_ONLINE_SESSION(Error, TEXT("EOS_Sessions_UpdateSession() failed with EOS result code (%s)"), ANSI_TO_TCHAR(EOS_EResult_ToString(Data->ResultCode)));
+				UE_LOG_ONLINE_SESSION(Error, TEXT("EOS_Sessions_UpdateSession() failed with EOS result code (%s)"), *LexToString(Data->ResultCode));
 
 				Session->SessionState = EOnlineSessionState::NoSession;
 
@@ -1701,7 +1729,7 @@ uint32 FOnlineSessionEOS::StartEOSSession(FNamedOnlineSession* Session)
 			bWasSuccessful = Data->ResultCode == EOS_EResult::EOS_Success;
 			if (!bWasSuccessful)
 			{
-				UE_LOG_ONLINE_SESSION(Error, TEXT("EOS_Sessions_StartSession() failed with EOS result code (%s)"), ANSI_TO_TCHAR(EOS_EResult_ToString(Data->ResultCode)));
+				UE_LOG_ONLINE_SESSION(Error, TEXT("EOS_Sessions_StartSession() failed with EOS result code (%s)"), *LexToString(Data->ResultCode));
 			}
 		}
 		else
@@ -1808,7 +1836,7 @@ uint32 FOnlineSessionEOS::UpdateEOSSession(FNamedOnlineSession* Session)
 	EOS_EResult ResultCode = EOS_Sessions_UpdateSessionModification(EOSSubsystem->SessionsHandle, &Options, &SessionModHandle);
 	if (ResultCode != EOS_EResult::EOS_Success)
 	{
-		UE_LOG_ONLINE_SESSION(Error, TEXT("EOS_Sessions_UpdateSessionModification() failed with EOS result code (%s)"), ANSI_TO_TCHAR(EOS_EResult_ToString(ResultCode)));
+		UE_LOG_ONLINE_SESSION(Error, TEXT("EOS_Sessions_UpdateSessionModification() failed with EOS result code (%s)"), *LexToString(ResultCode));
 		return ONLINE_FAIL;
 	}
 
@@ -1822,7 +1850,7 @@ uint32 FOnlineSessionEOS::UpdateEOSSession(FNamedOnlineSession* Session)
 			bWasSuccessful = Data->ResultCode == EOS_EResult::EOS_Success || Data->ResultCode == EOS_EResult::EOS_Sessions_OutOfSync;
 			if (!bWasSuccessful)
 			{
-				UE_LOG_ONLINE_SESSION(Warning, TEXT("EOS_Sessions_UpdateSession() failed with EOS result code (%s)"), ANSI_TO_TCHAR(EOS_EResult_ToString(Data->ResultCode)));
+				UE_LOG_ONLINE_SESSION(Warning, TEXT("EOS_Sessions_UpdateSession() failed with EOS result code (%s)"), *LexToString(Data->ResultCode));
 			}
 		}
 		else
@@ -1937,7 +1965,7 @@ uint32 FOnlineSessionEOS::EndEOSSession(FNamedOnlineSession* Session)
 			bWasSuccessful = Data->ResultCode == EOS_EResult::EOS_Success;
 			if (!bWasSuccessful)
 			{
-				UE_LOG_ONLINE_SESSION(Error, TEXT("EOS_Sessions_EndSession() failed with EOS result code (%s)"), ANSI_TO_TCHAR(EOS_EResult_ToString(Data->ResultCode)));
+				UE_LOG_ONLINE_SESSION(Error, TEXT("EOS_Sessions_EndSession() failed with EOS result code (%s)"), *LexToString(Data->ResultCode));
 			}
 		}
 		else
@@ -2055,7 +2083,7 @@ void FOnlineSessionEOS::EndSessionAnalytics()
 		EOS_EResult Result = EOS_Metrics_EndPlayerSession(EOSSubsystem->MetricsHandle, &Options);
 		if (Result != EOS_EResult::EOS_Success)
 		{
-			UE_LOG_ONLINE_SESSION(Error, TEXT("EOS_Metrics_EndPlayerSession() returned EOS result code (%s)"), ANSI_TO_TCHAR(EOS_EResult_ToString(Result)));
+			UE_LOG_ONLINE_SESSION(Error, TEXT("EOS_Metrics_EndPlayerSession() returned EOS result code (%s)"), *LexToString(Result));
 		}
 	}
 	else
@@ -2095,7 +2123,7 @@ uint32 FOnlineSessionEOS::DestroyEOSSession(FNamedOnlineSession* Session, const
 			bWasSuccessful = Data->ResultCode == EOS_EResult::EOS_Success;
 			if (!bWasSuccessful)
 			{
-				UE_LOG_ONLINE_SESSION(Error, TEXT("EOS_Sessions_DestroySession() failed with EOS result code (%s)"), ANSI_TO_TCHAR(EOS_EResult_ToString(Data->ResultCode)));
+				UE_LOG_ONLINE_SESSION(Error, TEXT("EOS_Sessions_DestroySession() failed with EOS result code (%s)"), *LexToString(Data->ResultCode));
 			}
 		}
 		else
@@ -2247,7 +2275,7 @@ bool FOnlineSessionEOS::FindSessionById(const FUniqueNetId& SearchingUserId, con
 	}
 	else
 	{
-		UE_LOG_ONLINE_SESSION(Warning, TEXT("[FOnlineSessionEOS::FindSessionById] CreateLobbySearch not successful. Finished with EOS_EResult %s"), ANSI_TO_TCHAR(EOS_EResult_ToString(CreateLobbySearchResult)));
+		UE_LOG_ONLINE_SESSION(Warning, TEXT("[FOnlineSessionEOS::FindSessionById] CreateLobbySearch not successful. Finished with EOS_EResult %s"), *LexToString(CreateLobbySearchResult));
 	}
 
 	return bResult;
@@ -2264,7 +2292,7 @@ void FOnlineSessionEOS::AddSearchAttribute(EOS_HSessionSearch SearchHandle, cons
 	EOS_EResult ResultCode = EOS_SessionSearch_SetParameter(SearchHandle, &Options);
 	if (ResultCode != EOS_EResult::EOS_Success)
 	{
-		UE_LOG_ONLINE_SESSION(Error, TEXT("EOS_SessionSearch_SetParameter() failed with EOS result code (%s)"), ANSI_TO_TCHAR(EOS_EResult_ToString(ResultCode)));
+		UE_LOG_ONLINE_SESSION(Error, TEXT("EOS_SessionSearch_SetParameter() failed with EOS result code (%s)"), *LexToString(ResultCode));
 	}
 }
 
@@ -2279,7 +2307,7 @@ void FOnlineSessionEOS::AddLobbySearchAttribute(EOS_HLobbySearch LobbySearchHand
 	EOS_EResult ResultCode = EOS_LobbySearch_SetParameter(LobbySearchHandle, &Options);
 	if (ResultCode != EOS_EResult::EOS_Success)
 	{
-		UE_LOG_ONLINE_SESSION(Error, TEXT("EOS_LobbySearch_SetParameter() failed with EOS result code (%s)"), ANSI_TO_TCHAR(EOS_EResult_ToString(ResultCode)));
+		UE_LOG_ONLINE_SESSION(Error, TEXT("EOS_LobbySearch_SetParameter() failed with EOS result code (%s)"), *LexToString(ResultCode));
 	}
 }
 
@@ -2395,21 +2423,27 @@ void FOnlineSessionEOS::CopyAttributes(const FSessionDetailsEOS& SessionHandle,
 
 void FOnlineSessionEOS::AddSearchResult(const TSharedRef<FSessionDetailsEOS>& SessionHandle, const TSharedRef<FOnlineSessionSearch>& SearchSettings)
 {
-	EOS_SessionDetails_Info* SessionInfo = nullptr;
+	EOS_SessionDetails_Info* EosSessionDetailsInfo = nullptr;
 	EOS_SessionDetails_CopyInfoOptions CopyOptions = { };
 	CopyOptions.ApiVersion = 1;
 	UE_EOS_CHECK_API_MISMATCH(EOS_SESSIONDETAILS_COPYINFO_API_LATEST, 1);
-	EOS_EResult CopyResult = EOS_SessionDetails_CopyInfo(SessionHandle->SessionDetailsHandle, &CopyOptions, &SessionInfo);
+	EOS_EResult CopyResult = EOS_SessionDetails_CopyInfo(SessionHandle->SessionDetailsHandle, &CopyOptions, &EosSessionDetailsInfo);
 	if (CopyResult == EOS_EResult::EOS_Success)
 	{
 		int32 Position = SearchSettings->SearchResults.AddZeroed();
 		FOnlineSessionSearchResult& SearchResult = SearchSettings->SearchResults[Position];
+
+
 		// This will set the host address and port
-		SearchResult.Session.SessionInfo = MakeShared<FOnlineSessionInfoEOS>(FOnlineSessionInfoEOS::Create(SessionInfo->HostAddress, FUniqueNetIdEOSSession::Create(SessionInfo->SessionId), SessionHandle));
+		TSharedPtr<FOnlineSessionInfoEOS> OnlineSessionInfo = MakeShared<FOnlineSessionInfoEOS>(FOnlineSessionInfoEOS::Create(FUniqueNetIdEOSSession::Create(EosSessionDetailsInfo->SessionId), SessionHandle));
+		SearchResult.Session.SessionInfo = OnlineSessionInfo;
 
-		CopySearchResult(*SessionHandle, SessionInfo, SearchResult.Session);
+		CopySearchResult(*SessionHandle, EosSessionDetailsInfo, SearchResult.Session);
 
-		EOS_SessionDetails_Info_Release(SessionInfo);
+		// CopySearchResult above will populate the settings so we can now read the port and construct the HostAddress.
+		OnlineSessionInfo->HostAddr = GetHostAddress(UTF8_TO_TCHAR(EosSessionDetailsInfo->HostAddress), SearchResult.Session.SessionSettings);
+
+		EOS_SessionDetails_Info_Release(EosSessionDetailsInfo);
 	}
 }
 
@@ -2426,7 +2460,7 @@ uint32 FOnlineSessionEOS::FindEOSSession(int32 SearchingPlayerNum, const TShared
 	EOS_EResult ResultCode = EOS_Sessions_CreateSessionSearch(EOSSubsystem->SessionsHandle, &HandleOptions, &SearchHandle);
 	if (ResultCode != EOS_EResult::EOS_Success)
 	{
-		UE_LOG_ONLINE_SESSION(Error, TEXT("EOS_Sessions_CreateSessionSearch() failed with EOS result code (%s)"), ANSI_TO_TCHAR(EOS_EResult_ToString(ResultCode)));
+		UE_LOG_ONLINE_SESSION(Error, TEXT("EOS_Sessions_CreateSessionSearch() failed with EOS result code (%s)"), *LexToString(ResultCode));
 		return ONLINE_FAIL;
 	}
 	// Store our search handle for use/cleanup later
@@ -2489,7 +2523,7 @@ uint32 FOnlineSessionEOS::FindEOSSession(int32 SearchingPlayerNum, const TShared
 		else
 		{
 			SearchSettings->SearchState = EOnlineAsyncTaskState::Failed;
-			UE_LOG_ONLINE_SESSION(Error, TEXT("EOS_SessionSearch_Find() failed with EOS result code (%s)"), ANSI_TO_TCHAR(EOS_EResult_ToString(Data->ResultCode)));
+			UE_LOG_ONLINE_SESSION(Error, TEXT("EOS_SessionSearch_Find() failed with EOS result code (%s)"), *LexToString(Data->ResultCode));
 		}
 		TriggerOnFindSessionsCompleteDelegates(bWasSuccessful);
 	};
@@ -2517,7 +2551,7 @@ void FOnlineSessionEOS::FindEOSSessionById(int32 LocalUserNum, const FUniqueNetI
 	EOS_EResult ResultCode = EOS_Sessions_CreateSessionSearch(EOSSubsystem->SessionsHandle, &HandleOptions, &SearchHandle);
 	if (ResultCode != EOS_EResult::EOS_Success)
 	{
-		UE_LOG_ONLINE_SESSION(Warning, TEXT("EOS_Sessions_CreateSessionSearch() failed with EOS result code (%s)"), ANSI_TO_TCHAR(EOS_EResult_ToString(ResultCode)));
+		UE_LOG_ONLINE_SESSION(Warning, TEXT("EOS_Sessions_CreateSessionSearch() failed with EOS result code (%s)"), *LexToString(ResultCode));
 		CompletionDelegate.ExecuteIfBound(LocalUserNum, false, FOnlineSessionSearchResult());
 		return;
 	}
@@ -2530,7 +2564,7 @@ void FOnlineSessionEOS::FindEOSSessionById(int32 LocalUserNum, const FUniqueNetI
 	ResultCode = EOS_SessionSearch_SetSessionId(SearchHandle, &Options);
 	if (ResultCode != EOS_EResult::EOS_Success)
 	{
-		UE_LOG_ONLINE_SESSION(Warning, TEXT("EOS_SessionSearch_SetSessionId() failed with EOS result code (%s)"), ANSI_TO_TCHAR(EOS_EResult_ToString(ResultCode)));
+		UE_LOG_ONLINE_SESSION(Warning, TEXT("EOS_SessionSearch_SetSessionId() failed with EOS result code (%s)"), *LexToString(ResultCode));
 		CompletionDelegate.ExecuteIfBound(LocalUserNum, false, FOnlineSessionSearchResult());
 		return;
 	}
@@ -2572,7 +2606,7 @@ void FOnlineSessionEOS::FindEOSSessionById(int32 LocalUserNum, const FUniqueNetI
 		else
 		{
 			LocalSessionSearch->SearchState = EOnlineAsyncTaskState::Failed;
-			UE_LOG_ONLINE_SESSION(Error, TEXT("EOS_SessionSearch_Find() failed with EOS result code (%s)"), ANSI_TO_TCHAR(EOS_EResult_ToString(Data->ResultCode)));
+			UE_LOG_ONLINE_SESSION(Error, TEXT("EOS_SessionSearch_Find() failed with EOS result code (%s)"), *LexToString(Data->ResultCode));
 		}
 
 		OnComplete.ExecuteIfBound(LocalUserNum, LocalSessionSearch->SearchState == EOnlineAsyncTaskState::Done, !LocalSessionSearch->SearchResults.IsEmpty() ? LocalSessionSearch->SearchResults.Last() : FOnlineSessionSearchResult());
@@ -2794,7 +2828,7 @@ uint32 FOnlineSessionEOS::JoinEOSSession(int32 PlayerNum, FNamedOnlineSession* S
 			}
 			else
 			{
-				UE_LOG_ONLINE_SESSION(Error, TEXT("EOS_Sessions_JoinSession() failed for session (%s) with EOS result code (%s)"), *SessionName.ToString(), ANSI_TO_TCHAR(EOS_EResult_ToString(Data->ResultCode)));
+				UE_LOG_ONLINE_SESSION(Error, TEXT("EOS_Sessions_JoinSession() failed for session (%s) with EOS result code (%s)"), *SessionName.ToString(), *LexToString(Data->ResultCode));
 
 				Session->SessionState = EOnlineSessionState::NoSession;
 
@@ -2873,7 +2907,7 @@ bool FOnlineSessionEOS::FindFriendSession(int32 LocalUserNum, const FUniqueNetId
 	}
 	else
 	{
-		UE_LOG_ONLINE_SESSION(Warning, TEXT("[FOnlineSessionEOS::FindFriendSession] CreateLobbySearch not successful. Finished with EOS_EResult %s"), ANSI_TO_TCHAR(EOS_EResult_ToString(CreateLobbySearchResult)));
+		UE_LOG_ONLINE_SESSION(Warning, TEXT("[FOnlineSessionEOS::FindFriendSession] CreateLobbySearch not successful. Finished with EOS_EResult %s"), *LexToString(CreateLobbySearchResult));
 		EOSSubsystem->ExecuteNextTick([this]()
 		{
 			TriggerOnFindSessionsCompleteDelegates(false);
@@ -2957,7 +2991,7 @@ bool FOnlineSessionEOS::SendLobbyInvite(FName SessionName, EOS_ProductUserId Sen
 		}
 		else
 		{
-			UE_LOG_ONLINE_SESSION(Warning, TEXT("[FOnlineSessionEOS::SendLobbyInvite] SendInvite not successful. Finished with EOS_EResult %s"), ANSI_TO_TCHAR(EOS_EResult_ToString(Data->ResultCode)));
+			UE_LOG_ONLINE_SESSION(Warning, TEXT("[FOnlineSessionEOS::SendLobbyInvite] SendInvite not successful. Finished with EOS_EResult %s"), *LexToString(Data->ResultCode));
 		}
 	};
 
@@ -2978,7 +3012,7 @@ bool FOnlineSessionEOS::SendEOSSessionInvite(FName SessionName, EOS_ProductUserI
 		bool bWasSuccessful = Data->ResultCode == EOS_EResult::EOS_Success;
 		if (!bWasSuccessful)
 		{
-			UE_LOG_ONLINE_SESSION(Error, TEXT("SendSessionInvite() failed for session (%s) with EOS result code (%s)"), *SessionName.ToString(), ANSI_TO_TCHAR(EOS_EResult_ToString(Data->ResultCode)));
+			UE_LOG_ONLINE_SESSION(Error, TEXT("SendSessionInvite() failed for session (%s) with EOS result code (%s)"), *SessionName.ToString(), *LexToString(Data->ResultCode));
 		}
 	};
 
@@ -3067,10 +3101,6 @@ static bool GetConnectStringFromSessionInfo(TSharedPtr<FOnlineSessionInfoEOS>& S
 	{
 		ConnectInfo = FString::Printf(TEXT("%s:%d"), *SessionInfo->HostAddr->ToString(false), PortOverride);
 	}
-	else if (SessionInfo->EOSAddress.Len() > 0)
-	{
-		ConnectInfo = SessionInfo->EOSAddress;
-	}
 	else
 	{
 		ConnectInfo = SessionInfo->HostAddr->ToString(true);
@@ -3196,7 +3226,7 @@ void FOnlineSessionEOS::UpdateOrAddLobbyMember(const FUniqueNetIdEOSLobbyRef& Lo
 			}
 			else
 			{
-				UE_LOG_ONLINE_SESSION(Warning, TEXT("[FOnlineSessionEOS::UpdateOrAddLobbyMember] EOS_LobbyDetails_CopyLobbyDetailsHandle not successful. Finished with EOS_EResult %s"), ANSI_TO_TCHAR(EOS_EResult_ToString(Result)));
+				UE_LOG_ONLINE_SESSION(Warning, TEXT("[FOnlineSessionEOS::UpdateOrAddLobbyMember] EOS_LobbyDetails_CopyLobbyDetailsHandle not successful. Finished with EOS_EResult %s"), *LexToString(Result));
 			}
 		}
 		else
@@ -3762,7 +3792,7 @@ void FOnlineSessionEOS::RemovePlayerFromSession(int32 LocalUserNum, FName Sessio
 			}
 			else
 			{
-				UE_LOG_ONLINE_SESSION(Warning, TEXT("[FOnlineSessionEOS::RemovePlayerFromSession] KickMember not successful. Finished with EOS_EResult %s"), ANSI_TO_TCHAR(EOS_EResult_ToString(Data->ResultCode)));
+				UE_LOG_ONLINE_SESSION(Warning, TEXT("[FOnlineSessionEOS::RemovePlayerFromSession] KickMember not successful. Finished with EOS_EResult %s"), *LexToString(Data->ResultCode));
 			}
 		};
 
@@ -3964,12 +3994,10 @@ uint32 FOnlineSessionEOS::CreateLobbySession(int32 HostingPlayerNum, FNamedOnlin
 
 				Session->SessionState = EOnlineSessionState::Pending;
 
-				// Because some platforms remap ports, we will use the ID of the name of the net driver to be our port instead
-				FName NetDriverName = GetDefault<UNetDriverEOS>()->NetDriverName;
-				FInternetAddrEOS TempAddr(LexToString(LocalProductUserId), SessionName.ToString(), FURL::UrlConfig.DefaultPort);
-				FString HostAddr = TempAddr.ToString(true);
+				TSharedPtr<FOnlineSessionInfoEOS> SessionInfo = MakeShared<FOnlineSessionInfoEOS>(FOnlineSessionInfoEOS::Create(FUniqueNetIdEOSLobby::Create(UTF8_TO_TCHAR(Data->LobbyId))));
+				SessionInfo->HostAddr = MakeShared<FInternetAddrEOS>(LocalProductUserId, SessionName.ToString(), FURL::UrlConfig.DefaultPort);
 
-				Session->SessionInfo = MakeShared<FOnlineSessionInfoEOS>(FOnlineSessionInfoEOS::Create(HostAddr, FUniqueNetIdEOSLobby::Create(UTF8_TO_TCHAR(Data->LobbyId))));
+				Session->SessionInfo = SessionInfo;
 
 #if WITH_EOSVOICECHAT
 				if (FEOSVoiceChatUser* VoiceChatUser = static_cast<FEOSVoiceChatUser*>(EOSSubsystem->GetEOSVoiceChatUserInterface(*LocalUserNetId)))
@@ -3984,7 +4012,7 @@ uint32 FOnlineSessionEOS::CreateLobbySession(int32 HostingPlayerNum, FNamedOnlin
 			}
 			else
 			{
-				UE_LOG_ONLINE_SESSION(Warning, TEXT("[FOnlineSessionEOS::CreateLobbySession] CreateLobby not successful. Finished with EOS_EResult %s"), ANSI_TO_TCHAR(EOS_EResult_ToString(Data->ResultCode)));
+				UE_LOG_ONLINE_SESSION(Warning, TEXT("[FOnlineSessionEOS::CreateLobbySession] CreateLobby not successful. Finished with EOS_EResult %s"), *LexToString(Data->ResultCode));
 
 				Session->SessionState = EOnlineSessionState::NoSession;
 
@@ -4048,7 +4076,7 @@ void FOnlineSessionEOS::DestroyLobbySessionOnCreationUpdateError(int32 LocalUser
 			}
 			else
 			{
-				UE_LOG_ONLINE_SESSION(Warning, TEXT("[FOnlineSessionEOS::DestroyLobbySessionOnCreationUpdateError] DestroyLobby not successful. Finished with EOS_EResult %s"), ANSI_TO_TCHAR(EOS_EResult_ToString(Data->ResultCode)));
+				UE_LOG_ONLINE_SESSION(Warning, TEXT("[FOnlineSessionEOS::DestroyLobbySessionOnCreationUpdateError] DestroyLobby not successful. Finished with EOS_EResult %s"), *LexToString(Data->ResultCode));
 			}
 
 #if WITH_EOSVOICECHAT
@@ -4087,13 +4115,12 @@ uint32 FOnlineSessionEOS::JoinLobbySession(int32 PlayerNum, FNamedOnlineSession*
 		FOnlineSessionInfoEOS* EOSSessionInfo = (FOnlineSessionInfoEOS*)(Session->SessionInfo.Get());
 		if (EOSSessionInfo->SessionId->IsValid())
 		{
+			// TODO why not just copy construct/assign?
 			const FOnlineSessionInfoEOS* SearchSessionInfo = (const FOnlineSessionInfoEOS*)(SearchSession->SessionInfo.Get());
 			EOSSessionInfo->HostAddr = SearchSessionInfo->HostAddr;
-			EOSSessionInfo->EOSAddress = SearchSessionInfo->EOSAddress;
 			EOSSessionInfo->SessionHandle = SearchSessionInfo->SessionHandle;
 			EOSSessionInfo->LobbyHandle = SearchSessionInfo->LobbyHandle;
 			EOSSessionInfo->SessionId = SearchSessionInfo->SessionId;
-			EOSSessionInfo->bIsFromClone = SearchSessionInfo->bIsFromClone;
 
 			Session->SessionState = EOnlineSessionState::Pending;
 			
@@ -4133,7 +4160,7 @@ uint32 FOnlineSessionEOS::JoinLobbySession(int32 PlayerNum, FNamedOnlineSession*
 					}
 					else
 					{
-						UE_LOG_ONLINE_SESSION(Warning, TEXT("[FOnlineSessionEOS::JoinLobbySession] JoinLobby not successful. Finished with EOS_EResult %s"), ANSI_TO_TCHAR(EOS_EResult_ToString(Data->ResultCode)));
+						UE_LOG_ONLINE_SESSION(Warning, TEXT("[FOnlineSessionEOS::JoinLobbySession] JoinLobby not successful. Finished with EOS_EResult %s"), *LexToString(Data->ResultCode));
 
 						Session->SessionState = EOnlineSessionState::NoSession;
 
@@ -4178,7 +4205,7 @@ void FOnlineSessionEOS::SetLobbyPermissionLevel(EOS_HLobbyModification LobbyModi
 	EOS_EResult ResultCode = EOS_LobbyModification_SetPermissionLevel(LobbyModificationHandle, &Options);
 	if (ResultCode != EOS_EResult::EOS_Success)
 	{
-		UE_LOG_ONLINE_SESSION(Warning, TEXT("[FOnlineSessionEOS::SetLobbyPermissionLevel] LobbyModification_SetPermissionLevel not successful. Finished with EOS_EResult %s"), ANSI_TO_TCHAR(EOS_EResult_ToString(ResultCode)));
+		UE_LOG_ONLINE_SESSION(Warning, TEXT("[FOnlineSessionEOS::SetLobbyPermissionLevel] LobbyModification_SetPermissionLevel not successful. Finished with EOS_EResult %s"), *LexToString(ResultCode));
 	}
 }
 
@@ -4194,7 +4221,7 @@ void FOnlineSessionEOS::SetLobbyMaxMembers(EOS_HLobbyModification LobbyModificat
 	EOS_EResult ResultCode = EOS_LobbyModification_SetMaxMembers(LobbyModificationHandle, &Options);
 	if (ResultCode != EOS_EResult::EOS_Success)
 	{
-		UE_LOG_ONLINE_SESSION(Warning, TEXT("[FOnlineSessionEOS::SetLobbyMaxMembers] LobbyModification_SetJoinInProgressAllowed not successful. Finished with EOS_EResult %s"), ANSI_TO_TCHAR(EOS_EResult_ToString(ResultCode)));
+		UE_LOG_ONLINE_SESSION(Warning, TEXT("[FOnlineSessionEOS::SetLobbyMaxMembers] LobbyModification_SetJoinInProgressAllowed not successful. Finished with EOS_EResult %s"), *LexToString(ResultCode));
 	}
 }
 
@@ -4209,7 +4236,7 @@ void FOnlineSessionEOS::AddLobbyAttribute(EOS_HLobbyModification LobbyModificati
 	EOS_EResult ResultCode = EOS_LobbyModification_AddAttribute(LobbyModificationHandle, &Options);
 	if (ResultCode != EOS_EResult::EOS_Success)
 	{
-		UE_LOG_ONLINE_SESSION(Error, TEXT("[FOnlineSessionEOS::AddLobbyAttribute] LobbyModification_AddAttribute for attribute name (%s) not successful. Finished with EOS_EResult %s"), *FString(Attribute->Key), ANSI_TO_TCHAR(EOS_EResult_ToString(ResultCode)));
+		UE_LOG_ONLINE_SESSION(Error, TEXT("[FOnlineSessionEOS::AddLobbyAttribute] LobbyModification_AddAttribute for attribute name (%s) not successful. Finished with EOS_EResult %s"), *FString(Attribute->Key), *LexToString(ResultCode));
 	}
 }
 
@@ -4224,7 +4251,7 @@ void FOnlineSessionEOS::AddLobbyMemberAttribute(EOS_HLobbyModification LobbyModi
 	EOS_EResult ResultCode = EOS_LobbyModification_AddMemberAttribute(LobbyModificationHandle, &Options);
 	if (ResultCode != EOS_EResult::EOS_Success)
 	{
-		UE_LOG_ONLINE_SESSION(Error, TEXT("[FOnlineSessionEOS::AddLobbyMemberAttribute] LobbyModification_AddMemberAttribute for attribute name (%s) not successful. Finished with EOS_EResult %s"), *FString(Attribute->Key), ANSI_TO_TCHAR(EOS_EResult_ToString(ResultCode)));
+		UE_LOG_ONLINE_SESSION(Error, TEXT("[FOnlineSessionEOS::AddLobbyMemberAttribute] LobbyModification_AddMemberAttribute for attribute name (%s) not successful. Finished with EOS_EResult %s"), *FString(Attribute->Key), *LexToString(ResultCode));
 	}
 }
 
@@ -4364,7 +4391,7 @@ uint32 FOnlineSessionEOS::UpdateLobbySession(FNamedOnlineSession* Session, const
 					bool bWasSuccessful = Data->ResultCode == EOS_EResult::EOS_Success || Data->ResultCode == EOS_EResult::EOS_Sessions_OutOfSync;
 					if (!bWasSuccessful)
 					{
-						UE_LOG_ONLINE_SESSION(Warning, TEXT("[FOnlineSessionEOS::UpdateLobbySession] UpdateLobby not successful. Finished with EOS_EResult %s"), ANSI_TO_TCHAR(EOS_EResult_ToString(Data->ResultCode)));
+						UE_LOG_ONLINE_SESSION(Warning, TEXT("[FOnlineSessionEOS::UpdateLobbySession] UpdateLobby not successful. Finished with EOS_EResult %s"), *LexToString(Data->ResultCode));
 					}
 
 					CompletionDelegate.ExecuteIfBound(SessionName, bWasSuccessful);
@@ -4384,7 +4411,7 @@ uint32 FOnlineSessionEOS::UpdateLobbySession(FNamedOnlineSession* Session, const
 		}
 		else
 		{
-			UE_LOG_ONLINE_SESSION(Warning, TEXT("[FOnlineSessionEOS::UpdateLobbySession] UpdateLobbyModification not successful. Finished with EOS_EResult %s"), ANSI_TO_TCHAR(EOS_EResult_ToString(LobbyModificationResult)));
+			UE_LOG_ONLINE_SESSION(Warning, TEXT("[FOnlineSessionEOS::UpdateLobbySession] UpdateLobbyModification not successful. Finished with EOS_EResult %s"), *LexToString(LobbyModificationResult));
 
 			EOSSubsystem->ExecuteNextTick([CompletionDelegate, SessionName = Session->SessionName]()
 				{
@@ -4454,7 +4481,7 @@ uint32 FOnlineSessionEOS::DestroyLobbySession(int32 LocalUserNum, FNamedOnlineSe
 				}
 				else
 				{
-					UE_LOG_ONLINE_SESSION(Warning, TEXT("[FOnlineSessionEOS::DestroyLobbySession] LeaveLobby not successful. Finished with EOS_EResult %s"), ANSI_TO_TCHAR(EOS_EResult_ToString(Data->ResultCode)));
+					UE_LOG_ONLINE_SESSION(Warning, TEXT("[FOnlineSessionEOS::DestroyLobbySession] LeaveLobby not successful. Finished with EOS_EResult %s"), *LexToString(Data->ResultCode));
 				}
 
 #if WITH_EOSVOICECHAT
@@ -4530,7 +4557,7 @@ uint32 FOnlineSessionEOS::FindLobbySession(int32 SearchingPlayerNum, const TShar
 	}
 	else
 	{
-		UE_LOG_ONLINE_SESSION(Warning, TEXT("[FOnlineSessionEOS::FindLobbySession] CreateLobbySearch not successful. Finished with EOS_EResult %s"), ANSI_TO_TCHAR(EOS_EResult_ToString(SearchResult)));
+		UE_LOG_ONLINE_SESSION(Warning, TEXT("[FOnlineSessionEOS::FindLobbySession] CreateLobbySearch not successful. Finished with EOS_EResult %s"), *LexToString(SearchResult));
 	}
 
 	return Result;
@@ -4582,7 +4609,7 @@ void FOnlineSessionEOS::StartLobbySearch(int32 SearchingPlayerNum, EOS_HLobbySea
 					}
 					else
 					{
-						UE_LOG_ONLINE_SESSION(Warning, TEXT("[FOnlineSessionEOS::StartLobbySearch::FLobbySearchFindCallback] LobbySearch_CopySearchResultByIndex not successful. Finished with EOS_EResult %s"), ANSI_TO_TCHAR(EOS_EResult_ToString(Data->ResultCode)));
+						UE_LOG_ONLINE_SESSION(Warning, TEXT("[FOnlineSessionEOS::StartLobbySearch::FLobbySearchFindCallback] LobbySearch_CopySearchResultByIndex not successful. Finished with EOS_EResult %s"), *LexToString(Data->ResultCode));
 					}
 				}
 
@@ -4616,7 +4643,7 @@ void FOnlineSessionEOS::StartLobbySearch(int32 SearchingPlayerNum, EOS_HLobbySea
 		}
 		else
 		{
-			UE_LOG_ONLINE_SESSION(Warning, TEXT("[FOnlineSessionEOS::StartLobbySearch::FLobbySearchFindCallback] LobbySearch_Find not successful. Finished with EOS_EResult %s"), ANSI_TO_TCHAR(EOS_EResult_ToString(Data->ResultCode)));
+			UE_LOG_ONLINE_SESSION(Warning, TEXT("[FOnlineSessionEOS::StartLobbySearch::FLobbySearchFindCallback] LobbySearch_Find not successful. Finished with EOS_EResult %s"), *LexToString(Data->ResultCode));
 
 			CurrentSessionSearch->SearchState = EOnlineAsyncTaskState::Failed;
 
@@ -4641,14 +4668,15 @@ void FOnlineSessionEOS::AddLobbySearchResult(const TSharedRef<FLobbyDetailsEOS>&
 		int32 Position = SearchSettings->SearchResults.AddZeroed();
 		FOnlineSessionSearchResult& SearchResult = SearchSettings->SearchResults[Position];
 		SearchResult.PingInMs = static_cast<int32>((FPlatformTime::Seconds() - SessionSearchStartInSeconds) * 1000);
-		
+
+		TSharedPtr<FOnlineSessionInfoEOS> SessionInfo = MakeShared<FOnlineSessionInfoEOS>(FOnlineSessionInfoEOS::Create(FUniqueNetIdEOSLobby::Create(UTF8_TO_TCHAR(LobbyDetailsInfo->LobbyId)), LobbyDetails));
+
 		// This will set the host address and port
 		// Because some platforms remap ports, we will use the ID of the name of the net driver to be our port instead
-		FName NetDriverName = GetDefault<UNetDriverEOS>()->NetDriverName;
-		FInternetAddrEOS TempAddr(LexToString(LobbyDetailsInfo->LobbyOwnerUserId), NetDriverName.ToString(), GetTypeHash(NetDriverName.ToString()));
-		FString HostAddr = TempAddr.ToString(true);
+		const FName NetDriverName = GetDefault<UNetDriverEOS>()->NetDriverName;
+		SessionInfo->HostAddr = MakeShared<FInternetAddrEOS>(LobbyDetailsInfo->LobbyOwnerUserId, NetDriverName.ToString(), GetTypeHash(NetDriverName.ToString()));
 
-		SearchResult.Session.SessionInfo = MakeShared<FOnlineSessionInfoEOS>(FOnlineSessionInfoEOS::Create(HostAddr, FUniqueNetIdEOSLobby::Create(UTF8_TO_TCHAR(LobbyDetailsInfo->LobbyId)), LobbyDetails));
+		SearchResult.Session.SessionInfo = SessionInfo;
 
 		// We copy the lobby data and settings, but not the member data (for search results)
 		CopyLobbyData(LobbyDetails, LobbyDetailsInfo, SearchResult.Session, false, Callback);
diff --git a/Engine/Plugins/Online/OnlineSubsystemEOS/Source/OnlineSubsystemEOS/Private/OnlineSubsystemEOSTypes.h b/Engine/Plugins/Online/OnlineSubsystemEOS/Source/OnlineSubsystemEOS/Private/OnlineSubsystemEOSTypes.h
--- a/Engine/Plugins/Online/OnlineSubsystemEOS/Source/OnlineSubsystemEOS/Private/OnlineSubsystemEOSTypes.h
+++ b/Engine/Plugins/Online/OnlineSubsystemEOS/Source/OnlineSubsystemEOS/Private/OnlineSubsystemEOSTypes.h
@@ -14,6 +14,11 @@
 #include "OnlineSubsystemTypes.h"
 #include "OnlineSubsystemEOSPackage.h" // IWYU pragma: keep
 
+#if defined(EOS_PLATFORM_BASE_FILE_NAME)
+#include EOS_PLATFORM_BASE_FILE_NAME
+#endif
+#include "eos_sessions_types.h"
+
 #define EOS_OSS_BUCKET_ID_STRING_LENGTH 60
 #define EOS_OSS_STRING_BUFFER_LENGTH 256 + 1 // 256 plus null terminator
 
@@ -122,6 +127,8 @@ private:
 #define USER_ATTR_COUNTRY TEXT("country")
 #define USER_ATTR_LANG TEXT("language")
 
+#define SESSION_ATTR_SERVERPORT TEXT("port")
+
 #if WITH_EOS_SDK
 
 /** Used to update all types of FOnlineUser classes, irrespective of leaf most class type */
@@ -551,8 +558,6 @@ private:
 	}
 };
 
-#include "eos_sessions_types.h"
-
 struct FSessionDetailsEOS : FNoncopyable
 {
 	EOS_HSessionDetails SessionDetailsHandle;
@@ -598,43 +603,31 @@ protected:
 
 PACKAGE_SCOPE:
 	/** Constructor */
-	FOnlineSessionInfoEOS();
-
-	FOnlineSessionInfoEOS(const FOnlineSessionInfoEOS& Src)
-		: FOnlineSessionInfo(Src)
-		, HostAddr(Src.HostAddr)
-		, SessionId(Src.SessionId)
-		, SessionHandle(Src.SessionHandle)
-		, LobbyHandle(Src.LobbyHandle)
-		, bIsFromClone(true)
-	{
-	}
-
-	FOnlineSessionInfoEOS(const FString& InHostIp, FUniqueNetIdStringRef UniqueNetId, const TSharedPtr<FSessionDetailsEOS>& InSessionHandle, const TSharedPtr<FLobbyDetailsEOS>& InLobbyHandle);
+	FOnlineSessionInfoEOS() = default;
+	
+	FOnlineSessionInfoEOS(const FOnlineSessionInfoEOS&) = default;
 
-	static FOnlineSessionInfoEOS Create(const FString& InHostIp, FUniqueNetIdStringRef UniqueNetId);
-	static FOnlineSessionInfoEOS Create(const FString& InHostIp, FUniqueNetIdStringRef UniqueNetId, const TSharedPtr<FSessionDetailsEOS>& InSessionHandle);
-	static FOnlineSessionInfoEOS Create(const FString& InHostIp, FUniqueNetIdStringRef UniqueNetId, const TSharedPtr<FLobbyDetailsEOS>& InLobbyHandle);
+	static FOnlineSessionInfoEOS Create(FUniqueNetIdStringRef UniqueNetId);
+	static FOnlineSessionInfoEOS Create(FUniqueNetIdStringRef UniqueNetId, const TSharedPtr<FSessionDetailsEOS>& SessionHandle);
+	static FOnlineSessionInfoEOS Create(FUniqueNetIdStringRef UniqueNetId, const TSharedPtr<FLobbyDetailsEOS>& LobbyHandle);
 
 	/**
 	 * Initialize LAN session
 	 */
 	void InitLAN(FOnlineSubsystemEOS* Subsystem);
 
-	FString EOSAddress;
 	/** The ip & port that the host is listening on (valid for LAN/GameServer) */
-	TSharedPtr<class FInternetAddr> HostAddr;
+	TSharedPtr<class FInternetAddr> HostAddr = nullptr;
 	/** Unique Id for this session */
-	FUniqueNetIdStringRef SessionId;
+	FUniqueNetIdStringRef SessionId = FUniqueNetIdString::EmptyId();
 	/** EOS session handle. The same handle can be shared between a local session and a search result. The struct type will call the release API automatically upon destruction */
-	TSharedPtr<FSessionDetailsEOS> SessionHandle;
+	TSharedPtr<FSessionDetailsEOS> SessionHandle = nullptr;
 	/** EOS lobby handle. The same handle can be shared between a local session and a search result. The struct type will call the release API automatically upon destruction */
-	TSharedPtr<FLobbyDetailsEOS> LobbyHandle;
-	/** Whether we should delete this handle or not */
-	bool bIsFromClone;
+	TSharedPtr<FLobbyDetailsEOS> LobbyHandle = nullptr;
 
 public:
-	virtual ~FOnlineSessionInfoEOS();
+	virtual ~FOnlineSessionInfoEOS() = default;
+
 	bool operator==(const FOnlineSessionInfoEOS& Other) const
 	{
 		return false;
