- Fix broken linux build due to parenthesis warning

diff --git a/Engine/Plugins/Online/OnlineSubsystemEOS/Source/OnlineSubsystemEOS/Private/OnlineSessionEOS.cpp b/Engine/Plugins/Online/OnlineSubsystemEOS/Source/OnlineSubsystemEOS/Private/OnlineSessionEOS.cpp
--- a/Engine/Plugins/Online/OnlineSubsystemEOS/Source/OnlineSubsystemEOS/Private/OnlineSessionEOS.cpp
+++ b/Engine/Plugins/Online/OnlineSubsystemEOS/Source/OnlineSubsystemEOS/Private/OnlineSessionEOS.cpp
@@ -1456,7 +1456,8 @@ TSharedPtr<class FInternetAddr> GetHostAddress(const FString& EosHostAddress, co
 	}
 	else
 	{
-		if (Result = ISocketSubsystem::Get(PLATFORM_SOCKETSUBSYSTEM)->GetAddressFromString(EosHostAddress))
+		Result = ISocketSubsystem::Get(PLATFORM_SOCKETSUBSYSTEM)->GetAddressFromString(EosHostAddress);
+		if (Result)
 		{
 			// Read port from settings, or fall back on DefaultPort.
 			int32 Port = FURL::UrlConfig.DefaultPort;
