- Fix various places calling EOS_EResult_ToString directly to call LexToString instead

diff --git a/Engine/Plugins/Online/OnlineServicesEOS/Source/Private/Online/CommerceEOS.cpp b/Engine/Plugins/Online/OnlineServicesEOS/Source/Private/Online/CommerceEOS.cpp
--- a/Engine/Plugins/Online/OnlineServicesEOS/Source/Private/Online/CommerceEOS.cpp
+++ b/Engine/Plugins/Online/OnlineServicesEOS/Source/Private/Online/CommerceEOS.cpp
@@ -228,7 +228,7 @@ TOnlineAsyncOpHandle<FCommerceCheckout> FCommerceEOS::Checkout(FCommerceCheckout
 		EOS_EResult Result = Data->ResultCode;
 		if (Result != EOS_EResult::EOS_Success)
 		{
-			UE_LOG(LogOnlineServices, Error, TEXT("EOS_Ecom_Checkout: failed with error (%s)"), ANSI_TO_TCHAR(EOS_EResult_ToString(Data->ResultCode)));
+			UE_LOG(LogOnlineServices, Error, TEXT("EOS_Ecom_Checkout: failed with error (%s)"), *LexToString(Data->ResultCode));
 			Op.SetError(Errors::FromEOSResult(Result));
 			return;
 		}
@@ -273,7 +273,7 @@ TOnlineAsyncOpHandle<FCommerceQueryTransactionEntitlements> FCommerceEOS::QueryT
 			EOS_EResult CopyTransactionResult = EOS_Ecom_CopyTransactionById(EcomHandle, &CopyTransactionOptions, &OutTransaction);
 			if (CopyTransactionResult != EOS_EResult::EOS_Success)
 			{
-				UE_LOG(LogOnlineServices, Error, TEXT("EOS_Ecom_CopyTransactionById: failed with error (%s)"), ANSI_TO_TCHAR(EOS_EResult_ToString(CopyTransactionResult)));
+				UE_LOG(LogOnlineServices, Error, TEXT("EOS_Ecom_CopyTransactionById: failed with error (%s)"), *LexToString(CopyTransactionResult));
 				Op.SetError(Errors::FromEOSResult(CopyTransactionResult));
 				return;
 			}
@@ -352,7 +352,7 @@ TOnlineAsyncOpHandle<FCommerceQueryEntitlements> FCommerceEOS::QueryEntitlements
 			const EOS_EResult Result = Data->ResultCode;
 			if (Result != EOS_EResult::EOS_Success)
 			{
-				UE_LOG(LogOnlineServices, Error, TEXT("EOS_Ecom_QueryEntitlements: failed with error (%s)"), ANSI_TO_TCHAR(EOS_EResult_ToString(Data->ResultCode)));
+				UE_LOG(LogOnlineServices, Error, TEXT("EOS_Ecom_QueryEntitlements: failed with error (%s)"), *LexToString(Data->ResultCode));
 				Op.SetError(Errors::FromEOSResult(Result));
 				return;
 			}
@@ -378,7 +378,7 @@ TOnlineAsyncOpHandle<FCommerceQueryEntitlements> FCommerceEOS::QueryEntitlements
 				EOS_EResult CopyResult = EOS_Ecom_CopyEntitlementByIndex(EcomHandle, &CopyOptions, &EosEntitlement);
 				if (CopyResult != EOS_EResult::EOS_Success && CopyResult != EOS_EResult::EOS_Ecom_EntitlementStale)
 				{
-					UE_LOG(LogOnlineServices, Error, TEXT("EOS_Ecom_CopyEntitlementByIndex: failed with error (%s) (proceeding with operation)"), ANSI_TO_TCHAR(EOS_EResult_ToString(CopyResult)));
+					UE_LOG(LogOnlineServices, Error, TEXT("EOS_Ecom_CopyEntitlementByIndex: failed with error (%s) (proceeding with operation)"), *LexToString(CopyResult));
 					continue;
 				}
 
@@ -441,7 +441,7 @@ TOnlineAsyncOpHandle<FCommerceRedeemEntitlement> FCommerceEOS::RedeemEntitlement
 		EOS_EResult Result = Data->ResultCode;
 		if (Result != EOS_EResult::EOS_Success)
 		{
-			UE_LOG(LogOnlineServices, Error, TEXT("EOS_Ecom_RedeemEntitlements: failed with error (%s)"), ANSI_TO_TCHAR(EOS_EResult_ToString(Data->ResultCode)));
+			UE_LOG(LogOnlineServices, Error, TEXT("EOS_Ecom_RedeemEntitlements: failed with error (%s)"), *LexToString(Data->ResultCode));
 			Op.SetError(Errors::FromEOSResult(Result));
 			return;
 		}
@@ -486,7 +486,7 @@ TOnlineAsyncOpHandle<FCommerceRetrieveS2SToken> FCommerceEOS::RetrieveS2SToken(F
 			EOS_EResult Result = Data->ResultCode;
 			if (Result != EOS_EResult::EOS_Success)
 			{
-				UE_LOG(LogOnlineServices, Error, TEXT("EOS_Ecom_QueryOwnershipToken: failed with error (%s)"), ANSI_TO_TCHAR(EOS_EResult_ToString(Data->ResultCode)));
+				UE_LOG(LogOnlineServices, Error, TEXT("EOS_Ecom_QueryOwnershipToken: failed with error (%s)"), *LexToString(Data->ResultCode));
 				Op.SetError(Errors::FromEOSResult(Result));
 				return;
 			}
diff --git a/Engine/Plugins/Online/OnlineSubsystemEOS/Source/OnlineSubsystemEOS/Private/OnlineAchievementsEOS.cpp b/Engine/Plugins/Online/OnlineSubsystemEOS/Source/OnlineSubsystemEOS/Private/OnlineAchievementsEOS.cpp
--- a/Engine/Plugins/Online/OnlineSubsystemEOS/Source/OnlineSubsystemEOS/Private/OnlineAchievementsEOS.cpp
+++ b/Engine/Plugins/Online/OnlineSubsystemEOS/Source/OnlineSubsystemEOS/Private/OnlineAchievementsEOS.cpp
@@ -157,13 +157,13 @@ void FOnlineAchievementsEOS::QueryAchievements(const FUniqueNetId& PlayerId, con
 				}
 				else
 				{
-					UE_LOG_ONLINE_ACHIEVEMENTS(Error, TEXT("EOS_Achievements_CopyPlayerAchievementByIndex() failed with error code (%s)"), ANSI_TO_TCHAR(EOS_EResult_ToString(Result)));
+					UE_LOG_ONLINE_ACHIEVEMENTS(Error, TEXT("EOS_Achievements_CopyPlayerAchievementByIndex() failed with error code (%s)"), *LexToString(Result));
 				}
 			}
 		}
 		else
 		{
-			UE_LOG_ONLINE_ACHIEVEMENTS(Error, TEXT("EOS_Achievements_QueryPlayerAchievements() failed with error code (%s)"), ANSI_TO_TCHAR(EOS_EResult_ToString(Data->ResultCode)));
+			UE_LOG_ONLINE_ACHIEVEMENTS(Error, TEXT("EOS_Achievements_QueryPlayerAchievements() failed with error code (%s)"), *LexToString(Data->ResultCode));
 		}
 		OnComplete.ExecuteIfBound(*LambdaPlayerId, bWasSuccessful);
 	};
@@ -237,13 +237,13 @@ void FOnlineAchievementsEOS::QueryAchievementDescriptions(const FUniqueNetId& Pl
 				}
 				else
 				{
-					UE_LOG_ONLINE_ACHIEVEMENTS(Error, TEXT("EOS_Achievements_CopyAchievementDefinitionByIndex() failed with error code (%s)"), ANSI_TO_TCHAR(EOS_EResult_ToString(Result)));
+					UE_LOG_ONLINE_ACHIEVEMENTS(Error, TEXT("EOS_Achievements_CopyAchievementDefinitionByIndex() failed with error code (%s)"), *LexToString(Result));
 				}
 			}
 		}
 		else
 		{
-			UE_LOG_ONLINE_ACHIEVEMENTS(Error, TEXT("EOS_Achievements_QueryDefinitions() failed with error code (%s)"), ANSI_TO_TCHAR(EOS_EResult_ToString(Data->ResultCode)));
+			UE_LOG_ONLINE_ACHIEVEMENTS(Error, TEXT("EOS_Achievements_QueryDefinitions() failed with error code (%s)"), *LexToString(Data->ResultCode));
 		}
 		OnComplete.ExecuteIfBound(*LambdaPlayerId, bWasSuccessful);
 	};
diff --git a/Engine/Plugins/Online/OnlineSubsystemEOS/Source/OnlineSubsystemEOS/Private/OnlineLeaderboardsEOS.cpp b/Engine/Plugins/Online/OnlineSubsystemEOS/Source/OnlineSubsystemEOS/Private/OnlineLeaderboardsEOS.cpp
--- a/Engine/Plugins/Online/OnlineSubsystemEOS/Source/OnlineSubsystemEOS/Private/OnlineLeaderboardsEOS.cpp
+++ b/Engine/Plugins/Online/OnlineSubsystemEOS/Source/OnlineSubsystemEOS/Private/OnlineLeaderboardsEOS.cpp
@@ -125,7 +125,7 @@ bool FOnlineLeaderboardsEOS::ReadLeaderboards(const TArray<FUniqueNetIdRef>& Pla
 		if (!bWasSuccessful)
 		{
 	
-			UE_LOG_ONLINE_LEADERBOARD(Error, TEXT("EOS_Leaderboards_QueryLeaderboardUserScores() failed with EOS result code (%s)"), ANSI_TO_TCHAR(EOS_EResult_ToString(Data->ResultCode)));
+			UE_LOG_ONLINE_LEADERBOARD(Error, TEXT("EOS_Leaderboards_QueryLeaderboardUserScores() failed with EOS result code (%s)"), *LexToString(Data->ResultCode));
 			QueryContext->ReadObject->ReadState = EOnlineAsyncTaskState::Failed;
 			TriggerOnLeaderboardReadCompleteDelegates(false);
 			return;
@@ -264,7 +264,7 @@ bool FOnlineLeaderboardsEOS::ReadLeaderboardsAroundRank(int32 Rank, uint32 Range
 		bool bWasSuccessful = Data->ResultCode == EOS_EResult::EOS_Success;
 		if (!bWasSuccessful)
 		{
-			UE_LOG_ONLINE_LEADERBOARD(Error, TEXT("EOS_Leaderboards_QueryLeaderboardRanks() failed with EOS result code (%s)"), ANSI_TO_TCHAR(EOS_EResult_ToString(Data->ResultCode)));
+			UE_LOG_ONLINE_LEADERBOARD(Error, TEXT("EOS_Leaderboards_QueryLeaderboardRanks() failed with EOS result code (%s)"), *LexToString(Data->ResultCode));
 			LambdaReadObject->ReadState = EOnlineAsyncTaskState::Failed;
 			TriggerOnLeaderboardReadCompleteDelegates(false);
 			return;
diff --git a/Engine/Plugins/Online/OnlineSubsystemEOS/Source/OnlineSubsystemEOS/Private/OnlineStatsEOS.cpp b/Engine/Plugins/Online/OnlineSubsystemEOS/Source/OnlineSubsystemEOS/Private/OnlineStatsEOS.cpp
--- a/Engine/Plugins/Online/OnlineSubsystemEOS/Source/OnlineSubsystemEOS/Private/OnlineStatsEOS.cpp
+++ b/Engine/Plugins/Online/OnlineSubsystemEOS/Source/OnlineSubsystemEOS/Private/OnlineStatsEOS.cpp
@@ -316,7 +316,7 @@ void FOnlineStatsEOS::WriteStats(EOS_ProductUserId LocalUserId, EOS_ProductUserI
 		bool bWasSuccessful = Data->ResultCode == EOS_EResult::EOS_Success;
 		if (!bWasSuccessful)
 		{
-			UE_LOG_ONLINE_STATS(Error, TEXT("EOS_Stats_IngestStat() failed with EOS result code (%s)"), ANSI_TO_TCHAR(EOS_EResult_ToString(Data->ResultCode)));
+			UE_LOG_ONLINE_STATS(Error, TEXT("EOS_Stats_IngestStat() failed with EOS result code (%s)"), *LexToString(Data->ResultCode));
 		}
 	};
 	EOS_Stats_IngestStat(EOSSubsystem->StatsHandle, &Options, CallbackObj, CallbackObj->GetCallbackPtr());
diff --git a/Engine/Plugins/Online/OnlineSubsystemEOS/Source/OnlineSubsystemEOS/Private/OnlineStoreEOS.cpp b/Engine/Plugins/Online/OnlineSubsystemEOS/Source/OnlineSubsystemEOS/Private/OnlineStoreEOS.cpp
--- a/Engine/Plugins/Online/OnlineSubsystemEOS/Source/OnlineSubsystemEOS/Private/OnlineStoreEOS.cpp
+++ b/Engine/Plugins/Online/OnlineSubsystemEOS/Source/OnlineSubsystemEOS/Private/OnlineStoreEOS.cpp
@@ -69,7 +69,7 @@ void FOnlineStoreEOS::QueryOffers(const FUniqueNetId& UserId, const FOnQueryOnli
 		EOS_EResult Result = Data->ResultCode;
 		if (Result != EOS_EResult::EOS_Success)
 		{
-			OnComplete.ExecuteIfBound(false, CachedOfferIds, EOS_EResult_ToString(Data->ResultCode));
+			OnComplete.ExecuteIfBound(false, CachedOfferIds, LexToString(Data->ResultCode));
 			return;
 		}
 
@@ -198,7 +198,7 @@ void FOnlineStoreEOS::Checkout(const FUniqueNetId& UserId, const FPurchaseChecko
 		const EOS_EResult CheckoutResult = Data->ResultCode;
 		if (CheckoutResult != EOS_EResult::EOS_Success)
 		{
-			UE_LOG_ONLINE(Error, TEXT("EOS_Ecom_Checkout: failed with error (%s)"), ANSI_TO_TCHAR(EOS_EResult_ToString(Data->ResultCode)));
+			UE_LOG_ONLINE(Error, TEXT("EOS_Ecom_Checkout: failed with error (%s)"), *LexToString(Data->ResultCode));
 			if (CheckoutResult == EOS_EResult::EOS_Canceled)
 			{
 				OnComplete.ExecuteIfBound(ONLINE_ERROR(EOnlineErrorResult::Canceled), MakeShared<FPurchaseReceipt>());
@@ -315,7 +315,7 @@ void FOnlineStoreEOS::QueryReceipts(const FUniqueNetId& UserId, bool bRestoreRec
 		EOS_EResult Result = Data->ResultCode;
 		if (Result != EOS_EResult::EOS_Success)
 		{
-			UE_LOG_ONLINE(Error, TEXT("EOS_Ecom_QueryEntitlements: failed with error (%s)"), ANSI_TO_TCHAR(EOS_EResult_ToString(Data->ResultCode)));
+			UE_LOG_ONLINE(Error, TEXT("EOS_Ecom_QueryEntitlements: failed with error (%s)"), *LexToString(Data->ResultCode));
 			OnComplete.ExecuteIfBound(ONLINE_ERROR(EOnlineErrorResult::Unknown));
 			return;
 		}
@@ -340,7 +340,7 @@ void FOnlineStoreEOS::QueryReceipts(const FUniqueNetId& UserId, bool bRestoreRec
 			EOS_EResult CopyResult = EOS_Ecom_CopyEntitlementByIndex(EOSSubsystem->EcomHandle, &CopyOptions, &Receipt);
 			if (CopyResult != EOS_EResult::EOS_Success && CopyResult != EOS_EResult::EOS_Ecom_EntitlementStale)
 			{
-				UE_LOG_ONLINE(Error, TEXT("EOS_Ecom_CopyEntitlementByIndex: failed with error (%s)"), ANSI_TO_TCHAR(EOS_EResult_ToString(CopyResult)));
+				UE_LOG_ONLINE(Error, TEXT("EOS_Ecom_CopyEntitlementByIndex: failed with error (%s)"), *LexToString(CopyResult));
 				continue;
 			}
 
@@ -401,7 +401,7 @@ void FOnlineStoreEOS::FinalizeReceiptValidationInfo(const FUniqueNetId& UserId,
 		EOS_EResult Result = Data->ResultCode;
 		if (Result != EOS_EResult::EOS_Success)
 		{
-			UE_LOG_ONLINE(Error, TEXT("EOS_Ecom_RedeemEntitlements: failed with error (%s)"), ANSI_TO_TCHAR(EOS_EResult_ToString(Data->ResultCode)));
+			UE_LOG_ONLINE(Error, TEXT("EOS_Ecom_RedeemEntitlements: failed with error (%s)"), *LexToString(Data->ResultCode));
 			OnComplete.ExecuteIfBound(ONLINE_ERROR(EOnlineErrorResult::Unknown), Info);
 			return;
 		}
diff --git a/Engine/Plugins/Online/OnlineSubsystemEOS/Source/OnlineSubsystemEOS/Private/OnlineTitleFileEOS.cpp b/Engine/Plugins/Online/OnlineSubsystemEOS/Source/OnlineSubsystemEOS/Private/OnlineTitleFileEOS.cpp
--- a/Engine/Plugins/Online/OnlineSubsystemEOS/Source/OnlineSubsystemEOS/Private/OnlineTitleFileEOS.cpp
+++ b/Engine/Plugins/Online/OnlineSubsystemEOS/Source/OnlineSubsystemEOS/Private/OnlineTitleFileEOS.cpp
@@ -109,7 +109,7 @@ void FOnlineTitleFileEOS::DeleteCachedFiles(bool bSkipEnumerated)
 		}
 		else
 		{
-			UE_LOG_ONLINE_TITLEFILE(Error, TEXT("EOS_TitleStorage_DeleteCache() failed with error code (%s)"), ANSI_TO_TCHAR(EOS_EResult_ToString(Data->ResultCode)));
+			UE_LOG_ONLINE_TITLEFILE(Error, TEXT("EOS_TitleStorage_DeleteCache() failed with error code (%s)"), *LexToString(Data->ResultCode));
 		}
 	};
 
@@ -205,14 +205,14 @@ bool FOnlineTitleFileEOS::EnumerateFiles(const FPagedQuery& Page)
 				}
 				else
 				{
-					ErrorStr = ANSI_TO_TCHAR(EOS_EResult_ToString(Result));
+					ErrorStr = LexToString(Result);
 					UE_LOG_ONLINE_TITLEFILE(Error, TEXT("EOS_TitleStorage_CopyFileMetadataAtIndex() failed with error code (%s)"), *ErrorStr);
 				}
 			}
 		}
 		else
 		{
-			ErrorStr = ANSI_TO_TCHAR(EOS_EResult_ToString(Data->ResultCode));
+			ErrorStr = LexToString(Data->ResultCode);
 			UE_LOG_ONLINE_TITLEFILE(Error, TEXT("EOS_TitleStorage_QueryFileList() failed with error code (%s)"), *ErrorStr);
 		}
 
@@ -319,7 +319,7 @@ bool FOnlineTitleFileEOS::ReadFile(const FString& FileName)
 				// If we fail to complete reading the file, discard it from the known files
 				FileSet.Remove(FString(ANSI_TO_TCHAR(Data->Filename)));
 
-				UE_LOG_ONLINE_TITLEFILE(Error, TEXT("EOS_TitleStorage_ReadFile() failed with error code (%s)"), ANSI_TO_TCHAR(EOS_EResult_ToString(Data->ResultCode)));
+				UE_LOG_ONLINE_TITLEFILE(Error, TEXT("EOS_TitleStorage_ReadFile() failed with error code (%s)"), *LexToString(Data->ResultCode));
 			}
 		}
 		else
diff --git a/Engine/Plugins/Online/OnlineSubsystemEOS/Source/OnlineSubsystemEOS/Private/OnlineUserCloudEOS.cpp b/Engine/Plugins/Online/OnlineSubsystemEOS/Source/OnlineSubsystemEOS/Private/OnlineUserCloudEOS.cpp
--- a/Engine/Plugins/Online/OnlineSubsystemEOS/Source/OnlineSubsystemEOS/Private/OnlineUserCloudEOS.cpp
+++ b/Engine/Plugins/Online/OnlineSubsystemEOS/Source/OnlineSubsystemEOS/Private/OnlineUserCloudEOS.cpp
@@ -233,13 +233,13 @@ void FOnlineUserCloudEOS::EnumerateUserFiles(const FUniqueNetId& UserId)
 					}
 					else
 					{
-						UE_LOG_ONLINE_CLOUD(Warning, TEXT("[FOnlineUserCloudEOS::EnumerateUserFiles] EOS_PlayerDataStorage_CopyFileMetadataAtIndex was not successful. Finished with error %s on file with index %d"), ANSI_TO_TCHAR(EOS_EResult_ToString(Result)), Index);
+						UE_LOG_ONLINE_CLOUD(Warning, TEXT("[FOnlineUserCloudEOS::EnumerateUserFiles] EOS_PlayerDataStorage_CopyFileMetadataAtIndex was not successful. Finished with error %s on file with index %d"), *LexToString(Result), Index);
 					}
 				}
 			}
 			else
 			{
-				UE_LOG_ONLINE_CLOUD(Warning, TEXT("[FOnlineUserCloudEOS::EnumerateUserFiles] EOS_PlayerDataStorage_QueryFileList was not successful. Finished with error %s"), ANSI_TO_TCHAR(EOS_EResult_ToString(Data->ResultCode)));
+				UE_LOG_ONLINE_CLOUD(Warning, TEXT("[FOnlineUserCloudEOS::EnumerateUserFiles] EOS_PlayerDataStorage_QueryFileList was not successful. Finished with error %s"), *LexToString(Data->ResultCode));
 			}
 		}
 		else
@@ -404,7 +404,7 @@ bool FOnlineUserCloudEOS::ReadUserFile(const FUniqueNetId& UserId, const FString
 					// If we fail to complete reading the file, discard it from the known files
 					FileSetsPerUser.Find(SharedUserId)->Remove(FileName);
 
-					UE_LOG_ONLINE_CLOUD(Warning, TEXT("[FOnlineUserCloudEOS::ReadUserFile] EOS_PlayerDataStorage_ReadFile was not successful for file %s. Finished with error %s"), *FileName, ANSI_TO_TCHAR(EOS_EResult_ToString(Data->ResultCode)));
+					UE_LOG_ONLINE_CLOUD(Warning, TEXT("[FOnlineUserCloudEOS::ReadUserFile] EOS_PlayerDataStorage_ReadFile was not successful for file %s. Finished with error %s"), *FileName, *LexToString(Data->ResultCode));
 				}
 			}
 			else
@@ -586,7 +586,7 @@ PRAGMA_ENABLE_DEPRECATION_WARNINGS
 					// If we fail to complete writing the file, discard it from the known files
 					FileSetsPerUser.Find(SharedUserId)->Remove(FileName);
 
-					UE_LOG_ONLINE_CLOUD(Warning, TEXT("[FOnlineUserCloudEOS::WriteUserFile] EOS_PlayerDataStorage_WriteFile was not successful. Finished with error %s"), ANSI_TO_TCHAR(EOS_EResult_ToString(Data->ResultCode)));
+					UE_LOG_ONLINE_CLOUD(Warning, TEXT("[FOnlineUserCloudEOS::WriteUserFile] EOS_PlayerDataStorage_WriteFile was not successful. Finished with error %s"), *LexToString(Data->ResultCode));
 				}
 			}
 			else
@@ -671,7 +671,7 @@ void FOnlineUserCloudEOS::CancelWriteUserFile(const FUniqueNetId& UserId, const
 				else
 				{
 					// Result code will be EOS_NoChange if request had already completed (can't be canceled), and EOS_AlreadyPending if it's already been canceled before (this is a final state for a canceled request and won't change over time)
-					UE_LOG_ONLINE_CLOUD(Warning, TEXT("[FOnlineUserCloudEOS::CancelWriteUserFile] EOS_PlayerDataStorageFileTransferRequest_CancelRequest was not successful. Finished with error %s"), ANSI_TO_TCHAR(EOS_EResult_ToString(Result)));
+					UE_LOG_ONLINE_CLOUD(Warning, TEXT("[FOnlineUserCloudEOS::CancelWriteUserFile] EOS_PlayerDataStorageFileTransferRequest_CancelRequest was not successful. Finished with error %s"), *LexToString(Result));
 				}
 			}
 			else
@@ -777,7 +777,7 @@ bool FOnlineUserCloudEOS::DeleteUserFile(const FUniqueNetId& UserId, const FStri
 			else
 			{
 				// File deletion operations can fail if the user does not own the file
-				UE_LOG_ONLINE_CLOUD(Warning, TEXT("[FOnlineUserCloudEOS::DeleteUserFile] EOS_PlayerDataStorage_DeleteFile was not successful for file %s. Finished with error %s"), *FileName, ANSI_TO_TCHAR(EOS_EResult_ToString(Data->ResultCode)));
+				UE_LOG_ONLINE_CLOUD(Warning, TEXT("[FOnlineUserCloudEOS::DeleteUserFile] EOS_PlayerDataStorage_DeleteFile was not successful for file %s. Finished with error %s"), *FileName, *LexToString(Data->ResultCode));
 			}
 
 			TriggerOnDeleteUserFileCompleteDelegates(bWasSuccessful, *UserIdRef, FileName);
diff --git a/Engine/Plugins/Online/OnlineSubsystemEOS/Source/OnlineSubsystemEOS/Private/UserManagerEOS.cpp b/Engine/Plugins/Online/OnlineSubsystemEOS/Source/OnlineSubsystemEOS/Private/UserManagerEOS.cpp
--- a/Engine/Plugins/Online/OnlineSubsystemEOS/Source/OnlineSubsystemEOS/Private/UserManagerEOS.cpp
+++ b/Engine/Plugins/Online/OnlineSubsystemEOS/Source/OnlineSubsystemEOS/Private/UserManagerEOS.cpp
@@ -720,7 +720,7 @@ void FUserManagerEOS::OnEOSAuthLoginComplete(int32 LocalUserNum, const FOnlineAc
 	{
 		auto TriggerLoginFailure = [this, LocalUserNum, LoginResultCode = Data->ResultCode]()
 		{
-			FString ErrorString = FString::Printf(TEXT("Login(%d) failed with EOS result code (%s)"), LocalUserNum, ANSI_TO_TCHAR(EOS_EResult_ToString(LoginResultCode)));
+			FString ErrorString = FString::Printf(TEXT("Login(%d) failed with EOS result code (%s)"), LocalUserNum, *LexToString(LoginResultCode));
 			UE_LOG_ONLINE(Warning, TEXT("%s"), *ErrorString);
 			TriggerOnLoginCompleteDelegates(LocalUserNum, false, *FUniqueNetIdEOS::EmptyId(), ErrorString);
 		};
@@ -791,7 +791,7 @@ void FUserManagerEOS::LinkEAS(int32 LocalUserNum, EOS_ContinuanceToken Token, co
 		}
 		else
 		{
-			FString ErrorString = FString::Printf(TEXT("Login(%d) failed with EOS result code (%s)"), LocalUserNum, ANSI_TO_TCHAR(EOS_EResult_ToString(Data->ResultCode)));
+			FString ErrorString = FString::Printf(TEXT("Login(%d) failed with EOS result code (%s)"), LocalUserNum, *LexToString(Data->ResultCode));
 			UE_LOG_ONLINE(Warning, TEXT("%s"), *ErrorString);
 			TriggerOnLoginCompleteDelegates(LocalUserNum, false, *FUniqueNetIdEOS::EmptyId(), ErrorString);
 		}
@@ -891,7 +891,7 @@ bool FUserManagerEOS::ConnectLoginNoEAS(int32 LocalUserNum, const FOnlineAccount
 					}
 					else
 					{
-						const FString ErrorString = FString::Printf(TEXT("ConnectLoginNoEAS(%d) failed with EOS result code (%s)"), LocalUserNum, ANSI_TO_TCHAR(EOS_EResult_ToString(Data->ResultCode)));
+						const FString ErrorString = FString::Printf(TEXT("ConnectLoginNoEAS(%d) failed with EOS result code (%s)"), LocalUserNum, *LexToString(Data->ResultCode));
 						UE_LOG_ONLINE(Warning, TEXT("%s"), *ErrorString);
 						TriggerOnLoginCompleteDelegates(LocalUserNum, false, *FUniqueNetIdEOS::EmptyId(), ErrorString);
 					}
@@ -967,7 +967,7 @@ bool FUserManagerEOS::ConnectLoginEAS(int32 LocalUserNum, EOS_EpicAccountId Acco
 							}
 							else
 							{
-								UE_LOG_ONLINE(Error, TEXT("ConnectLogin(%d) failed with EOS result code (%s)"), LocalUserNum, ANSI_TO_TCHAR(EOS_EResult_ToString(Data->ResultCode)));
+								UE_LOG_ONLINE(Error, TEXT("ConnectLogin(%d) failed with EOS result code (%s)"), LocalUserNum, *LexToString(Data->ResultCode));
 								Logout(LocalUserNum);
 							}
 						};
@@ -977,7 +977,7 @@ bool FUserManagerEOS::ConnectLoginEAS(int32 LocalUserNum, EOS_EpicAccountId Acco
 					}
 					else
 					{
-						UE_LOG_ONLINE(Error, TEXT("ConnectLogin(%d) failed with EOS result code (%s)"), LocalUserNum, ANSI_TO_TCHAR(EOS_EResult_ToString(CopyResult)));
+						UE_LOG_ONLINE(Error, TEXT("ConnectLogin(%d) failed with EOS result code (%s)"), LocalUserNum, *LexToString(CopyResult));
 						Logout(LocalUserNum);
 					}
 #if ADD_USER_LOGIN_INFO
@@ -1026,7 +1026,7 @@ void FUserManagerEOS::RefreshConnectLogin(int32 LocalUserNum)
 			{
 				if (Data->ResultCode != EOS_EResult::EOS_Success)
 				{
-					UE_LOG_ONLINE(Error, TEXT("Failed to refresh ConnectLogin(%d) failed with EOS result code (%s)"), LocalUserNum, ANSI_TO_TCHAR(EOS_EResult_ToString(Data->ResultCode)));
+					UE_LOG_ONLINE(Error, TEXT("Failed to refresh ConnectLogin(%d) failed with EOS result code (%s)"), LocalUserNum, *LexToString(Data->ResultCode));
 					Logout(LocalUserNum);
 				}
 			};
@@ -1067,7 +1067,7 @@ void FUserManagerEOS::RefreshConnectLogin(int32 LocalUserNum)
 					{
 						if (Data->ResultCode != EOS_EResult::EOS_Success)
 						{
-							UE_LOG_ONLINE(Error, TEXT("Failed to refresh ConnectLogin(%d) failed with EOS result code (%s)"), LocalUserNum, ANSI_TO_TCHAR(EOS_EResult_ToString(Data->ResultCode)));
+							UE_LOG_ONLINE(Error, TEXT("Failed to refresh ConnectLogin(%d) failed with EOS result code (%s)"), LocalUserNum, *LexToString(Data->ResultCode));
 							Logout(LocalUserNum);
 						}
 					};
@@ -1096,7 +1096,7 @@ void FUserManagerEOS::CreateConnectedLogin(int32 LocalUserNum, EOS_EpicAccountId
 		else
 		{
 // @todo joeg - logout?
-			FString ErrorString = FString::Printf(TEXT("Login(%d) failed with EOS result code (%s)"), LocalUserNum, ANSI_TO_TCHAR(EOS_EResult_ToString(Data->ResultCode)));
+			FString ErrorString = FString::Printf(TEXT("Login(%d) failed with EOS result code (%s)"), LocalUserNum, *LexToString(Data->ResultCode));
 			TriggerOnLoginCompleteDelegates(LocalUserNum, false, *FUniqueNetIdEOS::EmptyId(), ErrorString);
 		}
 	};
@@ -1563,7 +1563,7 @@ bool FUserManagerEOS::GetEpicAccountIdFromProductUserId(int32 LocalUserNum, cons
 	}
 	else
 	{
-		UE_LOG_ONLINE(Verbose, TEXT("[FUserManagerEOS::GetEpicAccountIdFromProductUserId] EOS_Connect_GetProductUserIdMapping not successful for ProductUserId (%s). Finished with EOS_EResult %s"), *LexToString(ProductUserId), ANSI_TO_TCHAR(EOS_EResult_ToString(Result)));
+		UE_LOG_ONLINE(Verbose, TEXT("[FUserManagerEOS::GetEpicAccountIdFromProductUserId] EOS_Connect_GetProductUserIdMapping not successful for ProductUserId (%s). Finished with EOS_EResult %s"), *LexToString(ProductUserId), *LexToString(Result));
 	}
 
 	return bResult;
@@ -1617,7 +1617,7 @@ void FUserManagerEOS::ResolveUniqueNetIds(int32 LocalUserNum, const TArray<EOS_P
 		{
 			if (Data->ResultCode != EOS_EResult::EOS_Success)
 			{
-				UE_LOG_ONLINE(Verbose, TEXT("[FUserManagerEOS::ResolveUniqueNetIds] EOS_Connect_QueryProductUserIdMappings not successful for user (%s). Finished with EOS_EResult %s."), *LexToString(Data->LocalUserId), ANSI_TO_TCHAR(EOS_EResult_ToString(Data->ResultCode)));
+				UE_LOG_ONLINE(Verbose, TEXT("[FUserManagerEOS::ResolveUniqueNetIds] EOS_Connect_QueryProductUserIdMappings not successful for user (%s). Finished with EOS_EResult %s."), *LexToString(Data->LocalUserId), *LexToString(Data->ResultCode));
 			}
 
 			for (const EOS_ProductUserId& ProductUserId : ProductUserIdsToResolve)
@@ -1894,7 +1894,7 @@ FString FUserManagerEOS::GetAuthToken(int32 LocalUserNum) const
 	}
 	else
 	{
-		UE_LOG_ONLINE(Verbose, TEXT("[FUserManagerEOS::GetAuthToken] EOS_Auth_CopyUserAuthToken failed with EOS result code (%s) for user (%d)"), ANSI_TO_TCHAR(EOS_EResult_ToString(CopyResult)), LocalUserNum);
+		UE_LOG_ONLINE(Verbose, TEXT("[FUserManagerEOS::GetAuthToken] EOS_Auth_CopyUserAuthToken failed with EOS result code (%s) for user (%d)"), *LexToString(CopyResult), LocalUserNum);
 	}
 
 	return FString();
@@ -1967,7 +1967,7 @@ bool FUserManagerEOS::ShowFriendsUI(int32 LocalUserNum)
 		}
 		else
 		{
-			UE_LOG_ONLINE_EXTERNALUI(Warning, TEXT("[FUserManagerEOS::ShowFriendsUI] EOS_UI_ShowFriends was not successful. Finished with error %s"), ANSI_TO_TCHAR(EOS_EResult_ToString(Data->ResultCode)));
+			UE_LOG_ONLINE_EXTERNALUI(Warning, TEXT("[FUserManagerEOS::ShowFriendsUI] EOS_UI_ShowFriends was not successful. Finished with error %s"), *LexToString(Data->ResultCode));
 		}
 	};
 
@@ -2343,7 +2343,7 @@ bool FUserManagerEOS::ReadFriendsList(int32 LocalUserNum, const FString& ListNam
 		}
 		else
 		{
-			const FString ErrorString = FString::Printf(TEXT("ReadFriendsList(%d) failed with EOS result code (%s)"), LocalUserNum, ANSI_TO_TCHAR(EOS_EResult_ToString(Result)));
+			const FString ErrorString = FString::Printf(TEXT("ReadFriendsList(%d) failed with EOS result code (%s)"), LocalUserNum, *LexToString(Result));
 			ProcessReadFriendsListComplete(LocalUserNum, false, ErrorString);
 		}
 	};
@@ -2450,7 +2450,7 @@ bool FUserManagerEOS::SendInvite(int32 LocalUserNum, const FUniqueNetId& FriendI
 		bool bWasSuccessful = Data->ResultCode == EOS_EResult::EOS_Success;
 		if (!bWasSuccessful)
 		{
-			ErrorString = FString::Printf(TEXT("Failed to send invite for user (%d) to player (%s) with result code (%s)"), LocalUserNum, *NetIdRef->ToString(), ANSI_TO_TCHAR(EOS_EResult_ToString(Data->ResultCode)));
+			ErrorString = FString::Printf(TEXT("Failed to send invite for user (%d) to player (%s) with result code (%s)"), LocalUserNum, *NetIdRef->ToString(), *LexToString(Data->ResultCode));
 		}
 
 		Delegate.ExecuteIfBound(LocalUserNum, bWasSuccessful, *NetIdRef, ListName, ErrorString);
@@ -2495,7 +2495,7 @@ bool FUserManagerEOS::AcceptInvite(int32 LocalUserNum, const FUniqueNetId& Frien
 		bool bWasSuccessful = Data->ResultCode == EOS_EResult::EOS_Success;
 		if (!bWasSuccessful)
 		{
-			ErrorString = FString::Printf(TEXT("Failed to accept invite for user (%d) from friend (%s) with result code (%s)"), LocalUserNum, *NetIdRef->ToString(), ANSI_TO_TCHAR(EOS_EResult_ToString(Data->ResultCode)));
+			ErrorString = FString::Printf(TEXT("Failed to accept invite for user (%d) from friend (%s) with result code (%s)"), LocalUserNum, *NetIdRef->ToString(), *LexToString(Data->ResultCode));
 		}
 		Delegate.ExecuteIfBound(LocalUserNum, bWasSuccessful, *NetIdRef, ListName, ErrorString);
 	};
@@ -2875,7 +2875,7 @@ void FUserManagerEOS::SetPresence(const FUniqueNetId& UserId, const FOnlineUserP
 	EOS_EResult SetRichTextResult = EOS_PresenceModification_SetRawRichText(ChangeHandle, &TextOptions);
 	if (SetRichTextResult != EOS_EResult::EOS_Success)
 	{
-		UE_LOG_ONLINE(Error, TEXT("EOS_PresenceModification_SetRawRichText() failed with result code (%s)"), ANSI_TO_TCHAR(EOS_EResult_ToString(SetRichTextResult)));
+		UE_LOG_ONLINE(Error, TEXT("EOS_PresenceModification_SetRawRichText() failed with result code (%s)"), *LexToString(SetRichTextResult));
 	}
 
 	TArray<FPresenceStrings, TInlineAllocator<EOS_PRESENCE_DATA_MAX_KEYS>> RawStrings;
@@ -3038,7 +3038,7 @@ void FUserManagerEOS::UpdatePresence(int32 LocalUserNum, EOS_EpicAccountId Accou
 	}
 	else
 	{
-		UE_LOG_ONLINE(Error, TEXT("Failed to copy presence data with error code (%s)"), ANSI_TO_TCHAR(EOS_EResult_ToString(CopyResult)));
+		UE_LOG_ONLINE(Error, TEXT("Failed to copy presence data with error code (%s)"), *LexToString(CopyResult));
 	}
 }
 
@@ -3276,7 +3276,7 @@ bool FUserManagerEOS::QueryUserIdMapping(const FUniqueNetId& UserId, const FStri
 		}
 		else
 		{
-			ErrorString = FString::Printf(TEXT("QueryUserIdMapping(%d, '%s') failed with EOS result code (%s)"), DefaultLocalUser, *DisplayNameOrEmail, ANSI_TO_TCHAR(EOS_EResult_ToString(Result)));
+			ErrorString = FString::Printf(TEXT("QueryUserIdMapping(%d, '%s') failed with EOS result code (%s)"), DefaultLocalUser, *DisplayNameOrEmail, *LexToString(Result));
 		}
 		Delegate.ExecuteIfBound(false, *FUniqueNetIdEOS::EmptyId(), DisplayNameOrEmail, *FUniqueNetIdEOS::EmptyId(), ErrorString);
 	};
@@ -3403,7 +3403,7 @@ bool FUserManagerEOS::QueryExternalIdMappings(const FUniqueNetId& UserId, const
 			}
 			else
 			{
-				ErrorString = FString::Printf(TEXT("EOS_Connect_QueryExternalAccountMappings() failed with result code (%s)"), ANSI_TO_TCHAR(EOS_EResult_ToString(Result)));
+				ErrorString = FString::Printf(TEXT("EOS_Connect_QueryExternalAccountMappings() failed with result code (%s)"), *LexToString(Result));
 			}
 
 			// Mark all queries as complete
diff --git a/Engine/Plugins/Online/SocketSubsystemEOS/Source/SocketSubsystemEOS/Private/SocketEOS.cpp b/Engine/Plugins/Online/SocketSubsystemEOS/Source/SocketSubsystemEOS/Private/SocketEOS.cpp
--- a/Engine/Plugins/Online/SocketSubsystemEOS/Source/SocketSubsystemEOS/Private/SocketEOS.cpp
+++ b/Engine/Plugins/Online/SocketSubsystemEOS/Source/SocketSubsystemEOS/Private/SocketEOS.cpp
@@ -119,8 +119,8 @@ bool FSocketEOS::Close()
 
 		EOS_EResult Result = EOS_P2P_CloseConnections(SocketSubsystem.GetP2PHandle(), &Options);
 
-		UE_LOG(LogSocketSubsystemEOS, Log, TEXT("Closing socket (%s) with result (%s)"), *LocalAddress.ToString(true), ANSI_TO_TCHAR(EOS_EResult_ToString(Result)));
-		NP_LOG(TEXT("[%s] - Closing socket (%s) with result (%s)\r\n"), GetLogPrefix(), *LocalAddress.ToString(true), ANSI_TO_TCHAR(EOS_EResult_ToString(Result)));
+		UE_LOG(LogSocketSubsystemEOS, Log, TEXT("Closing socket (%s) with result (%s)"), *LocalAddress.ToString(true), *LexToString(Result));
+		NP_LOG(TEXT("[%s] - Closing socket (%s) with result (%s)\r\n"), GetLogPrefix(), *LocalAddress.ToString(true), *LexToString(Result));
 
 		ClosedRemotes.Empty();
 	}
@@ -244,8 +244,8 @@ bool FSocketEOS::Listen(int32)
 			}
 			else
 			{
-				UE_LOG(LogSocketSubsystemEOS, Error, TEXT("EOS_P2P_AcceptConnection from (%s) on socket (%s) failed with (%s)"), *RemoteUser, UTF8_TO_TCHAR(Info->SocketId->SocketName), ANSI_TO_TCHAR(EOS_EResult_ToString(AcceptResult)));
-				NP_LOG(TEXT("[%s] - EOS_P2P_AcceptConnection from (%s) on socket (%s) failed with (%s)\r\n"), GetLogPrefix(), *RemoteUser, UTF8_TO_TCHAR(Info->SocketId->SocketName), ANSI_TO_TCHAR(EOS_EResult_ToString(AcceptResult)));
+				UE_LOG(LogSocketSubsystemEOS, Error, TEXT("EOS_P2P_AcceptConnection from (%s) on socket (%s) failed with (%s)"), *RemoteUser, UTF8_TO_TCHAR(Info->SocketId->SocketName), *LexToString(AcceptResult));
+				NP_LOG(TEXT("[%s] - EOS_P2P_AcceptConnection from (%s) on socket (%s) failed with (%s)\r\n"), GetLogPrefix(), *RemoteUser, UTF8_TO_TCHAR(Info->SocketId->SocketName), *LexToString(AcceptResult));
 			}
 		}
 		else
@@ -292,7 +292,7 @@ bool FSocketEOS::HasPendingData(uint32& PendingDataSize)
 	}
 	if (Result != EOS_EResult::EOS_Success)
 	{
-		UE_LOG(LogSocketSubsystemEOS, Warning, TEXT("Unable to check for data on address (%s) result code = (%s)"), *LocalAddress.ToString(true), ANSI_TO_TCHAR(EOS_EResult_ToString(Result)));
+		UE_LOG(LogSocketSubsystemEOS, Warning, TEXT("Unable to check for data on address (%s) result code = (%s)"), *LocalAddress.ToString(true), *LexToString(Result));
 
 		// @todo joeg - map EOS codes to UE4's
 		SocketSubsystem.SetLastSocketError(ESocketErrors::SE_EINVAL);
@@ -406,10 +406,10 @@ bool FSocketEOS::SendTo(const uint8* Data, int32 Count, int32& OutBytesSent, con
 	Options.DataLengthBytes = Count;
 	Options.Data = Data;
 	EOS_EResult Result = EOS_P2P_SendPacket(SocketSubsystem.GetP2PHandle(), &Options);
-	NP_LOG(TEXT("[%s] - EOS_P2P_SendPacket() to (%s) result code = (%s)\r\n"), GetLogPrefix(), *Destination.ToString(true), ANSI_TO_TCHAR(EOS_EResult_ToString(Result)));
+	NP_LOG(TEXT("[%s] - EOS_P2P_SendPacket() to (%s) result code = (%s)\r\n"), GetLogPrefix(), *Destination.ToString(true), *LexToString(Result));
 	if (Result != EOS_EResult::EOS_Success)
 	{
-		UE_LOG(LogSocketSubsystemEOS, Error, TEXT("Unable to send data to (%s) result code = (%s)"), *Destination.ToString(true), ANSI_TO_TCHAR(EOS_EResult_ToString(Result)));
+		UE_LOG(LogSocketSubsystemEOS, Error, TEXT("Unable to send data to (%s) result code = (%s)"), *Destination.ToString(true), *LexToString(Result));
 
 		// @todo joeg - map EOS codes to UE4's
 		SocketSubsystem.SetLastSocketError(ESocketErrors::SE_EINVAL);
@@ -465,7 +465,7 @@ bool FSocketEOS::RecvFrom(uint8* Data, int32 BufferSize, int32& BytesRead, FInte
 	EOS_P2P_SocketId SocketId;
 	
 	EOS_EResult Result = EOS_P2P_ReceivePacket(SocketSubsystem.GetP2PHandle(), &Options, &RemoteUserId, &SocketId, &Channel, Data, (uint32*)&BytesRead);
-	NP_LOG(TEXT("[%s] - EOS_P2P_ReceivePacket() for user (%s) and channel (%d) with result code = (%s)\r\n"), GetLogPrefix(), *MakeStringFromProductUserId(LocalAddress.GetLocalUserId()), Channel, ANSI_TO_TCHAR(EOS_EResult_ToString(Result)));
+	NP_LOG(TEXT("[%s] - EOS_P2P_ReceivePacket() for user (%s) and channel (%d) with result code = (%s)\r\n"), GetLogPrefix(), *MakeStringFromProductUserId(LocalAddress.GetLocalUserId()), Channel, *LexToString(Result));
 	if (Result == EOS_EResult::EOS_NotFound)
 	{
 		// No data to read
@@ -474,7 +474,7 @@ bool FSocketEOS::RecvFrom(uint8* Data, int32 BufferSize, int32& BytesRead, FInte
 	}
 	else if (Result != EOS_EResult::EOS_Success)
 	{
-		UE_LOG(LogSocketSubsystemEOS, Error, TEXT("Unable to receive data result code = (%s)"), ANSI_TO_TCHAR(EOS_EResult_ToString(Result)));
+		UE_LOG(LogSocketSubsystemEOS, Error, TEXT("Unable to receive data result code = (%s)"), *LexToString(Result));
 
 		// @todo joeg - map EOS codes to UE4's
 		SocketSubsystem.SetLastSocketError(ESocketErrors::SE_EINVAL);
@@ -643,10 +643,10 @@ bool FSocketEOS::Close(const FInternetAddrEOS& RemoteAddress)
 	Options.SocketId = &SocketId;
 
 	EOS_EResult Result = EOS_P2P_CloseConnection(SocketSubsystem.GetP2PHandle(), &Options);
-	NP_LOG(TEXT("[%s] - EOS_P2P_CloseConnection() with remote address RemoteAddress (%s) result code (%s)\r\n"), GetLogPrefix(), *RemoteAddress.ToString(true), ANSI_TO_TCHAR(EOS_EResult_ToString(Result)));
+	NP_LOG(TEXT("[%s] - EOS_P2P_CloseConnection() with remote address RemoteAddress (%s) result code (%s)\r\n"), GetLogPrefix(), *RemoteAddress.ToString(true), *LexToString(Result));
 	if (Result != EOS_EResult::EOS_Success)
 	{
-		UE_LOG(LogSocketSubsystemEOS, Error, TEXT("Unable to close socket with remote address RemoteAddress (%s) due to error (%s)"), *RemoteAddress.ToString(true), ANSI_TO_TCHAR(EOS_EResult_ToString(Result)));
+		UE_LOG(LogSocketSubsystemEOS, Error, TEXT("Unable to close socket with remote address RemoteAddress (%s) due to error (%s)"), *RemoteAddress.ToString(true), *LexToString(Result));
 
 		// @todo joeg - map EOS codes to UE4's
 		SocketSubsystem.SetLastSocketError(ESocketErrors::SE_EINVAL);
diff --git a/Engine/Plugins/Online/SocketSubsystemEOS/Source/SocketSubsystemEOS/Private/SocketSubsystemEOS.cpp b/Engine/Plugins/Online/SocketSubsystemEOS/Source/SocketSubsystemEOS/Private/SocketSubsystemEOS.cpp
--- a/Engine/Plugins/Online/SocketSubsystemEOS/Source/SocketSubsystemEOS/Private/SocketSubsystemEOS.cpp
+++ b/Engine/Plugins/Online/SocketSubsystemEOS/Source/SocketSubsystemEOS/Private/SocketSubsystemEOS.cpp
@@ -103,7 +103,7 @@ bool FSocketSubsystemEOS::Init(FString& Error)
 		EOS_EResult Result = EOS_P2P_SetRelayControl(P2PHandle, &Options);
 		if (Result != EOS_EResult::EOS_Success && Result != EOS_EResult::EOS_NoChange)
 		{
-			UE_LOG(LogSocketSubsystemEOS, Warning, TEXT("[FSocketSubsystemEOS::Init] EOS_P2P_SetRelayControl failed with result [%s]"), UTF8_TO_TCHAR(EOS_EResult_ToString(Result)));
+			UE_LOG(LogSocketSubsystemEOS, Warning, TEXT("[FSocketSubsystemEOS::Init] EOS_P2P_SetRelayControl failed with result [%s]"), *LexToString(Result));
 		}
 	}
 #endif
