Added support to specify the type of auth in EOS external mapping queries

diff --git a/Engine/Plugins/Online/EOSShared/Source/EOSShared/Private/EOSShared.cpp b/Engine/Plugins/Online/EOSShared/Source/EOSShared/Private/EOSShared.cpp
--- a/Engine/Plugins/Online/EOSShared/Source/EOSShared/Private/EOSShared.cpp
+++ b/Engine/Plugins/Online/EOSShared/Source/EOSShared/Private/EOSShared.cpp
@@ -170,6 +170,73 @@ const TCHAR* LexToString(const EOS_Presence_EStatus PresenceStatus)
 	}
 }
 
+bool LexFromString(EOS_EExternalAccountType& OutEnum, const TCHAR* InString)
+{
+	if (FCString::Stricmp(InString, TEXT("Amazon")) == 0)
+	{
+		OutEnum = EOS_EExternalAccountType::EOS_EAT_AMAZON;
+	}
+	else if (FCString::Stricmp(InString, TEXT("Apple")) == 0)
+	{
+		OutEnum = EOS_EExternalAccountType::EOS_EAT_APPLE;
+	}
+	else if (FCString::Stricmp(InString, TEXT("Discord")) == 0)
+	{
+		OutEnum = EOS_EExternalAccountType::EOS_EAT_DISCORD;
+	}
+	else if (FCString::Stricmp(InString, TEXT("Epic")) == 0)
+	{
+		OutEnum = EOS_EExternalAccountType::EOS_EAT_EPIC;
+	}
+	else if (FCString::Stricmp(InString, TEXT("GOG")) == 0)
+	{
+		OutEnum = EOS_EExternalAccountType::EOS_EAT_GOG;
+	}
+	else if (FCString::Stricmp(InString, TEXT("Google")) == 0)
+	{
+		OutEnum = EOS_EExternalAccountType::EOS_EAT_GOOGLE;
+	}
+	else if (FCString::Stricmp(InString, TEXT("ItchIO")) == 0)
+	{
+		OutEnum = EOS_EExternalAccountType::EOS_EAT_ITCHIO;
+	}
+	else if (FCString::Stricmp(InString, TEXT("Nintendo")) == 0)
+	{
+		OutEnum = EOS_EExternalAccountType::EOS_EAT_NINTENDO;
+	}
+	else if (FCString::Stricmp(InString, TEXT("Oculus")) == 0)
+	{
+		OutEnum = EOS_EExternalAccountType::EOS_EAT_OCULUS;
+	}
+	else if (FCString::Stricmp(InString, TEXT("OpenID")) == 0)
+	{
+		OutEnum = EOS_EExternalAccountType::EOS_EAT_OPENID;
+	}
+	else if (FCString::Stricmp(InString, TEXT("PSN")) == 0)
+	{
+		OutEnum = EOS_EExternalAccountType::EOS_EAT_PSN;
+	}
+	else if (FCString::Stricmp(InString, TEXT("Steam")) == 0)
+	{
+		OutEnum = EOS_EExternalAccountType::EOS_EAT_STEAM;
+	}
+	else if (FCString::Stricmp(InString, TEXT("UPlay")) == 0)
+	{
+		OutEnum = EOS_EExternalAccountType::EOS_EAT_UPLAY;
+	}
+	else if (FCString::Stricmp(InString, TEXT("XBL")) == 0)
+	{
+		OutEnum = EOS_EExternalAccountType::EOS_EAT_XBL;
+	}
+	else
+	{
+		checkNoEntry();
+		return false;
+	}
+
+	return true;
+}
+
 bool LexFromString(EOS_EExternalCredentialType& OutEnum, const TCHAR* InString)
 {
 	if (FCString::Stricmp(InString, TEXT("Steam")) == 0)
diff --git a/Engine/Plugins/Online/EOSShared/Source/EOSShared/Public/EOSShared.h b/Engine/Plugins/Online/EOSShared/Source/EOSShared/Public/EOSShared.h
--- a/Engine/Plugins/Online/EOSShared/Source/EOSShared/Public/EOSShared.h
+++ b/Engine/Plugins/Online/EOSShared/Source/EOSShared/Public/EOSShared.h
@@ -86,6 +86,7 @@ EOSSHARED_API const TCHAR* LexToString(const EOS_ENetworkStatus NetworkStatus);
 EOSSHARED_API const TCHAR* LexToString(const EOS_Presence_EStatus PresenceStatus);
 
 EOSSHARED_API bool LexFromString(EOS_EAuthScopeFlags& OutEnum, const FStringView InString);
+EOSSHARED_API bool LexFromString(EOS_EExternalAccountType& OutEnum, const TCHAR* InString);
 EOSSHARED_API bool LexFromString(EOS_EExternalCredentialType& OutEnum, const TCHAR* InString);
 EOSSHARED_API bool LexFromString(EOS_ELoginCredentialType& OutEnum, const TCHAR* InString);
 EOSSHARED_API bool LexFromString(EOS_ERTCBackgroundMode& OutEnum, const TCHAR* InString);
diff --git a/Engine/Plugins/Online/OnlineSubsystemEOS/Source/OnlineSubsystemEOS/Private/UserManagerEOS.cpp b/Engine/Plugins/Online/OnlineSubsystemEOS/Source/OnlineSubsystemEOS/Private/UserManagerEOS.cpp
--- a/Engine/Plugins/Online/OnlineSubsystemEOS/Source/OnlineSubsystemEOS/Private/UserManagerEOS.cpp
+++ b/Engine/Plugins/Online/OnlineSubsystemEOS/Source/OnlineSubsystemEOS/Private/UserManagerEOS.cpp
@@ -1214,7 +1214,7 @@ bool FUserManagerEOS::Logout(int32 LocalUserNum)
 	EOS_Auth_Logout(EOSSubsystem->AuthHandle, &LogoutOptions, CallbackObj, CallbackObj->GetCallbackPtr());
 
 	return true;
-}
+} 
 
 bool FUserManagerEOS::AutoLogin(int32 LocalUserNum)
 {
@@ -1678,7 +1678,7 @@ void FUserManagerEOS::ResolveUniqueNetIds(int32 LocalUserNum, const TArray<EOS_E
 
 	if (!EpicAccountIdsToResolve.IsEmpty())
 	{
-		QueryExternalIdMappings(*GetLocalUniqueNetIdEOS(LocalUserNum), FExternalIdQueryOptions(), EpicAccountIdsToResolve, FOnQueryExternalIdMappingsComplete::CreateLambda([this, ResolvedUniqueNetIds = MoveTemp(ResolvedUniqueNetIds), Callback](bool bWasSuccessful, const FUniqueNetId& UserId, const FExternalIdQueryOptions& QueryOptions, const TArray<FString>& ExternalIds, const FString& Error) mutable
+		QueryExternalIdMappings(*GetLocalUniqueNetIdEOS(LocalUserNum), FExternalIdQueryOptions(TEXT("Epic"), false), EpicAccountIdsToResolve, FOnQueryExternalIdMappingsComplete::CreateLambda([this, ResolvedUniqueNetIds = MoveTemp(ResolvedUniqueNetIds), Callback](bool bWasSuccessful, const FUniqueNetId& UserId, const FExternalIdQueryOptions& QueryOptions, const TArray<FString>& ExternalIds, const FString& Error) mutable
 			{
 				if (bWasSuccessful)
 				{
@@ -3292,7 +3292,7 @@ bool FUserManagerEOS::QueryUserIdMapping(const FUniqueNetId& UserId, const FStri
 struct FQueryByStringIdsOptions :
 	public EOS_Connect_QueryExternalAccountMappingsOptions
 {
-	FQueryByStringIdsOptions(const uint32 InNumStringIds, EOS_ProductUserId InLocalUserId) :
+	FQueryByStringIdsOptions(const uint32 InNumStringIds, EOS_ProductUserId InLocalUserId, FString InAccountIdTypeStr) :
 		EOS_Connect_QueryExternalAccountMappingsOptions()
 	{
 		PointerArray.AddZeroed(InNumStringIds);
@@ -3302,7 +3302,10 @@ struct FQueryByStringIdsOptions :
 		}
 		ApiVersion = 1;
 		UE_EOS_CHECK_API_MISMATCH(EOS_CONNECT_QUERYEXTERNALACCOUNTMAPPINGS_API_LATEST, 1);
-		AccountIdType = EOS_EExternalAccountType::EOS_EAT_EPIC;
+		if (!LexFromString(AccountIdType, *InAccountIdTypeStr))
+		{
+			UE_LOG_ONLINE(Warning, TEXT("[%hs] Unable to parse AccountIdType [%s]"), __FUNCTION__, *InAccountIdTypeStr);
+		}
 		ExternalAccountIds = (const char**)PointerArray.GetData();
 		ExternalAccountIdCount = InNumStringIds;
 		LocalUserId = InLocalUserId;
@@ -3365,7 +3368,7 @@ bool FUserManagerEOS::QueryExternalIdMappings(const FUniqueNetId& UserId, const
 		const uint32 BatchSize = FMath::Min(ExternalIds.Num() - BatchSrcOffset, MaxBatchSize);
 
 		// Build an options up per batch
-		FQueryByStringIdsOptions Options(BatchSize, LocalUserId);
+		FQueryByStringIdsOptions Options(BatchSize, LocalUserId, QueryOptions.AuthType);
 		for (uint32 DestIdx = 0, SrcIdx = BatchSrcOffset; DestIdx < BatchSize; DestIdx++, SrcIdx++)
 		{
 			FCStringAnsi::Strncpy(Options.PointerArray[DestIdx], TCHAR_TO_UTF8(*ExternalIds[SrcIdx]), EOS_CONNECT_EXTERNAL_ACCOUNT_ID_MAX_LENGTH+1);
