- Add a UEOSSetting entry to select the SteamTokenType to request of OSSSteam.
- This enables licensees to upgrade to the "WebApi" token type, either passing an explicit remote service identity or using the default one configured on OSSSteam.

diff --git a/Engine/Plugins/Online/OnlineSubsystemEOS/Source/OnlineSubsystemEOS/Private/EOSSettings.cpp b/Engine/Plugins/Online/OnlineSubsystemEOS/Source/OnlineSubsystemEOS/Private/EOSSettings.cpp
--- a/Engine/Plugins/Online/OnlineSubsystemEOS/Source/OnlineSubsystemEOS/Private/EOSSettings.cpp
+++ b/Engine/Plugins/Online/OnlineSubsystemEOS/Source/OnlineSubsystemEOS/Private/EOSSettings.cpp
@@ -130,7 +130,8 @@ FEOSArtifactSettings ParseArtifactSettingsFromConfigString(const FString& RawLin
 }
 
 FEOSSettings::FEOSSettings()
-	: RTCBackgroundMode(EOS_ERTCBackgroundMode::EOS_RTCBM_KeepRoomsAlive)
+	: SteamTokenType(TEXT("Session"))
+	, RTCBackgroundMode(EOS_ERTCBackgroundMode::EOS_RTCBM_KeepRoomsAlive)
 	, TickBudgetInMilliseconds(0)
 	, TitleStorageReadChunkLength(0)
 	, bEnableOverlay(false)
@@ -172,6 +173,7 @@ const FEOSSettings& UEOSSettings::ManualGetSettings()
 
 		GConfig->GetString(INI_SECTION, TEXT("CacheDir"), CachedSettings->CacheDir, GEngineIni);
 		GConfig->GetString(INI_SECTION, TEXT("DefaultArtifactName"), CachedSettings->DefaultArtifactName, GEngineIni);
+		GConfig->GetString(INI_SECTION, TEXT("SteamTokenType"), CachedSettings->SteamTokenType, GEngineIni);
 		CachedSettings->RTCBackgroundMode = EOS_ERTCBackgroundMode::EOS_RTCBM_KeepRoomsAlive;
 		FString RTCBackgroundModeStr;
 		GConfig->GetString(INI_SECTION, TEXT("RTCBackgroundMode"), RTCBackgroundModeStr, GEngineIni);
@@ -205,6 +207,7 @@ FEOSSettings UEOSSettings::ToNative() const
 
 	Native.CacheDir = CacheDir;
 	Native.DefaultArtifactName = DefaultArtifactName;
+	Native.SteamTokenType = SteamTokenType;
 	Native.RTCBackgroundMode = EOS_ERTCBackgroundMode::EOS_RTCBM_KeepRoomsAlive;
 	if (!RTCBackgroundMode.IsEmpty())
 	{
diff --git a/Engine/Plugins/Online/OnlineSubsystemEOS/Source/OnlineSubsystemEOS/Private/UserManagerEOS.cpp b/Engine/Plugins/Online/OnlineSubsystemEOS/Source/OnlineSubsystemEOS/Private/UserManagerEOS.cpp
--- a/Engine/Plugins/Online/OnlineSubsystemEOS/Source/OnlineSubsystemEOS/Private/UserManagerEOS.cpp
+++ b/Engine/Plugins/Online/OnlineSubsystemEOS/Source/OnlineSubsystemEOS/Private/UserManagerEOS.cpp
@@ -103,6 +103,12 @@ static inline EOS_EExternalCredentialType ToEOS_EExternalCredentialType(FName OS
 #if PLATFORM_DESKTOP
 	if (OSSName == STEAM_SUBSYSTEM)
 	{
+		FEOSSettings Settings = UEOSSettings::GetSettings();
+		if (Settings.SteamTokenType == TEXT("App"))
+		{
+			return EOS_EExternalCredentialType::EOS_ECT_STEAM_APP_TICKET;
+		}
+		// Session, WebApi, and WebApi:remoteserviceidentity are all "Session" tickets.
 		return EOS_EExternalCredentialType::EOS_ECT_STEAM_SESSION_TICKET;
 	}
 #endif
@@ -343,10 +349,10 @@ void FUserManagerEOS::GetPlatformAuthToken(int32 LocalUserNum, const FOnGetLinke
 	}
 
 	FString TokenType;
-	// TODO config map of OSS -> token type?
 	if (PlatformOSS->GetSubsystemName() == STEAM_SUBSYSTEM)
 	{
-		TokenType = TEXT("Session");
+		FEOSSettings Settings = UEOSSettings::GetSettings();
+		TokenType = Settings.SteamTokenType;
 	}
 
 	// Request the auth token from the platform
diff --git a/Engine/Plugins/Online/OnlineSubsystemEOS/Source/OnlineSubsystemEOS/Public/EOSSettings.h b/Engine/Plugins/Online/OnlineSubsystemEOS/Source/OnlineSubsystemEOS/Public/EOSSettings.h
--- a/Engine/Plugins/Online/OnlineSubsystemEOS/Source/OnlineSubsystemEOS/Public/EOSSettings.h
+++ b/Engine/Plugins/Online/OnlineSubsystemEOS/Source/OnlineSubsystemEOS/Public/EOSSettings.h
@@ -71,6 +71,7 @@ struct FEOSSettings
 
 	FString CacheDir;
 	FString DefaultArtifactName;
+	FString SteamTokenType;
 	EOS_ERTCBackgroundMode RTCBackgroundMode;
 	int32 TickBudgetInMilliseconds;
 	int32 TitleStorageReadChunkLength;
@@ -172,6 +173,19 @@ public:
 	UPROPERTY(Config, EditAnywhere, BlueprintReadOnly, Category="Crossplay Settings")
 	bool bMirrorPresenceToEAS = false;
 
+	/**
+	 * When running with Steam, defines what TokenType OSSEOS will request from OSSSteam to login with.
+	 * Please see EOS documentation at https://dev.epicgames.com/docs/dev-portal/identity-provider-management#steam for more information.
+	 * Note the default is currently "Session" but this is deprecated. Please migrate to WebApi.
+	 * Possible values:
+	 *     "App" -> [DEPRECATED] Use Steam Encryption Application Tickets from ISteamUser::GetEncryptedAppTicket.
+	 *     "Session" -> [DEPRECATED] Use Steam Auth Session Tickets from ISteamUser::GetAuthSessionTicket.
+	 *     "WebApi" -> Use Steam Auth Tickets from ISteamUser::GetAuthTicketForWebApi, using the default remote service identity configured for OSSSteam.
+	 *     "WebApi:<remoteserviceidentity>" -> Use Steam Auth Tickets from ISteamUser::GetAuthTicketForWebApi, using an explicit remote service identity.
+	 */
+	UPROPERTY(Config, EditAnywhere, BlueprintReadOnly, Category="Crossplay Settings")
+	FString SteamTokenType = TEXT("Session");
+
 	/** Get the settings for the selected artifact */
 	static bool GetSelectedArtifactSettings(FEOSArtifactSettings& OutSettings);
 
