EOS+Platform compatability change.
Puts the EOS lobby ID into the platfrom session info using the key expected by Platform, and removes spurious check on invite result.
This resolves session creation fail and enables platfrom invites and join by presence.

diff --git a/Engine/Plugins/Online/OnlineSubsystemEOS/Source/OnlineSubsystemEOSPlus/Private/OnlineSessionEOSPlus.cpp b/Engine/Plugins/Online/OnlineSubsystemEOS/Source/OnlineSubsystemEOSPlus/Private/OnlineSessionEOSPlus.cpp
--- a/Engine/Plugins/Online/OnlineSubsystemEOS/Source/OnlineSubsystemEOSPlus/Private/OnlineSessionEOSPlus.cpp
+++ b/Engine/Plugins/Online/OnlineSubsystemEOS/Source/OnlineSubsystemEOSPlus/Private/OnlineSessionEOSPlus.cpp
@@ -4,6 +4,7 @@
 #include "Misc/Guid.h"
 #include "OnlineSubsystem.h"
 #include "OnlineSubsystemEOSPlus.h"
+#include "Online/OnlineSessionNames.h"
 #include "OnlineUserEOSPlus.h"
 #include "EOSSettings.h"
 
@@ -137,7 +138,7 @@ void FOnlineSessionEOSPlus::OnLoginComplete(int32 LocalUserNum, bool bWasSuccess
 
 void FOnlineSessionEOSPlus::OnSessionUserInviteAcceptedBase(const bool bWasSuccessful, const int32 ControllerId, FUniqueNetIdPtr UserId, const FOnlineSessionSearchResult& InviteResult)
 {
-	if (!bWasSuccessful || !UserId.IsValid() || !InviteResult.IsValid())
+	if (!bWasSuccessful || !UserId.IsValid())
 	{
 		return;
 	}
@@ -145,7 +146,7 @@ void FOnlineSessionEOSPlus::OnSessionUserInviteAcceptedBase(const bool bWasSucce
 	if (bUseEOSSessions)
 	{
 		FString SessionIdStr;
-		InviteResult.Session.SessionSettings.Get(TEXT("EOSSessionId"), SessionIdStr);
+		InviteResult.Session.SessionSettings.Get(SETTING_CUSTOM_JOIN_INFO, SessionIdStr);
 		if (SessionIdStr.IsEmpty())
 		{
 			UE_LOG_ONLINE(Error, TEXT("Failed to get EOS session id from base session"));
@@ -222,7 +223,7 @@ void FOnlineSessionEOSPlus::OnSessionInviteReceivedBase(const FUniqueNetId& User
 	if (bUseEOSSessions)
 	{
 		FString SessionIdStr;
-		InviteResult.Session.SessionSettings.Get(TEXT("EOSSessionId"), SessionIdStr);
+		InviteResult.Session.SessionSettings.Get(SETTING_CUSTOM_JOIN_INFO, SessionIdStr);
 		if (SessionIdStr.IsEmpty())
 		{
 			UE_LOG_ONLINE(Error, TEXT("Failed to get EOS session id from base session"));
@@ -278,8 +279,6 @@ void FOnlineSessionEOSPlus::OnCreateSessionComplete(FName SessionName, bool bWas
 	if (BaseSession != nullptr)
 	{
 		FNamedOnlineSession* EOSSession = EOSSessionInterface->GetNamedSession(SessionName);
-
-		BaseSession->SessionSettings.Set(TEXT("EOSSessionId"), EOSSession->SessionInfo->GetSessionId().ToString(), EOnlineDataAdvertisementType::ViaOnlineService);
 		BaseSessionInterface->UpdateSession(SessionName, BaseSession->SessionSettings, true);
 	}
 
@@ -398,7 +397,8 @@ bool FOnlineSessionEOSPlus::CreateSession(const FUniqueNetId& HostingPlayerId, F
 						if (Settings != nullptr)
 						{
 							FUniqueNetIdEOSPlusPtr Id = GetNetIdPlus(HostingPlayerIdStr);
-
+							FNamedOnlineSession* EOSSession = EOSSessionInterface->GetNamedSession(SessionName);
+							Settings->Set(SETTING_CUSTOM_JOIN_INFO, EOSSession->SessionInfo->GetSessionId().ToString(), EOnlineDataAdvertisementType::ViaOnlineService);
 							// Mirror in the base interface
 							BaseSessionInterface->CreateSession(*Id->GetBaseNetId(), SessionName, *Settings);
 
@@ -449,7 +449,11 @@ bool FOnlineSessionEOSPlus::UpdateSession(FName SessionName, FOnlineSessionSetti
 	}
 
 #if CREATE_MIRROR_PLATFORM_SESSION
-	return BaseSessionInterface->UpdateSession(SessionName, UpdatedSessionSettings, bShouldRefreshOnlineData);
+	FNamedOnlineSession* EOSSession = EOSSessionInterface->GetNamedSession(SessionName);
+	FNamedOnlineSession* BaseSession = BaseSessionInterface->GetNamedSession(SessionName);
+	BaseSession->SessionSettings = UpdatedSessionSettings;
+	BaseSession->SessionSettings.Set(SETTING_CUSTOM_JOIN_INFO, EOSSession->SessionInfo->GetSessionId().ToString(), EOnlineDataAdvertisementType::ViaOnlineService);
+	return BaseSessionInterface->UpdateSession(SessionName, BaseSession->SessionSettings, bShouldRefreshOnlineData);
 #else
 	return true;
 #endif
