Changes to GetFriendsList method to comply with interface requirements.
Removed the DisplayName check as it was being used to check if the user info query was still ongoing, and causing valid friend entries to not be returned.
Added error cases when the user is not found or the list name passed is not valid.

diff --git a/Engine/Plugins/Online/OnlineSubsystem/Source/Public/Interfaces/OnlineFriendsInterface.h b/Engine/Plugins/Online/OnlineSubsystem/Source/Public/Interfaces/OnlineFriendsInterface.h
--- a/Engine/Plugins/Online/OnlineSubsystem/Source/Public/Interfaces/OnlineFriendsInterface.h
+++ b/Engine/Plugins/Online/OnlineSubsystem/Source/Public/Interfaces/OnlineFriendsInterface.h
@@ -51,6 +51,34 @@ namespace EFriendsLists
 		}
 		return TEXT("");
 	}
+
+	/** @return whether or not the passed string was successfully converted to a valid enum value */
+	inline bool FromString(EFriendsLists::Type& OutEnum, const TCHAR* InString)
+	{
+		if (FCString::Stricmp(InString, TEXT("default")) == 0)
+		{
+			OutEnum = EFriendsLists::Type::Default;
+		}
+		else if (FCString::Stricmp(InString, TEXT("onlinePlayers")) == 0)
+		{
+			OutEnum = EFriendsLists::Type::OnlinePlayers;
+		}
+		else if (FCString::Stricmp(InString, TEXT("inGamePlayers")) == 0)
+		{
+			OutEnum = EFriendsLists::Type::InGamePlayers;
+		}
+		else if (FCString::Stricmp(InString, TEXT("inGameAndSessionPlayers")) == 0)
+		{
+			OutEnum = EFriendsLists::Type::InGameAndSessionPlayers;
+		}
+		else
+		{
+			checkNoEntry();
+			return false;
+		}
+
+		return true;
+	}
 }
 
 /*
diff --git a/Engine/Plugins/Online/OnlineSubsystemEOS/Source/OnlineSubsystemEOS/Private/UserManagerEOS.cpp b/Engine/Plugins/Online/OnlineSubsystemEOS/Source/OnlineSubsystemEOS/Private/UserManagerEOS.cpp
--- a/Engine/Plugins/Online/OnlineSubsystemEOS/Source/OnlineSubsystemEOS/Private/UserManagerEOS.cpp
+++ b/Engine/Plugins/Online/OnlineSubsystemEOS/Source/OnlineSubsystemEOS/Private/UserManagerEOS.cpp
@@ -573,9 +573,14 @@ void FUserManagerEOS::LoginViaExternalAuth(int32 LocalUserNum, const FOnlineAcco
 		}));
 }
 
+bool FUserManagerEOS::IsLocalUserValid(int32 LocalUserNum) const
+{
+	return LocalUsers.IsValidIndex(LocalUserNum);
+}
+
 FLocalUserEOS& FUserManagerEOS::GetLocalUserChecked(int32 LocalUserNum)
 {
-	check(LocalUsers.IsValidIndex(LocalUserNum));
+	check(IsLocalUserValid(LocalUserNum));
 	return LocalUsers[LocalUserNum];
 }
 
@@ -2562,23 +2567,29 @@ bool FUserManagerEOS::DeleteFriend(int32 LocalUserNum, const FUniqueNetId& Frien
 
 bool FUserManagerEOS::GetFriendsList(int32 LocalUserNum, const FString& ListName, TArray<TSharedRef<FOnlineFriend>>& OutFriends)
 {
+	if (!IsLocalUserValid(LocalUserNum))
+	{
+		return false;
+	}
+
+	EFriendsLists::Type FriendsListType;
+	if (!EFriendsLists::FromString(FriendsListType, *ListName))
+	{
+		return false;
+	}
+
 	OutFriends.Reset();
 
 	for (FOnlineFriendEOSRef Friend : GetLocalUserChecked(LocalUserNum).FriendsList->GetList())
 	{
 		const FOnlineUserPresence& Presence = Friend->GetPresence();
 		// See if they only want online only
-		if (ListName == EFriendsLists::ToString(EFriendsLists::OnlinePlayers) && !Presence.bIsOnline)
+		if (FriendsListType == EFriendsLists::OnlinePlayers && !Presence.bIsOnline)
 		{
 			continue;
 		}
 		// Of if they only want friends playing this game
-		else if (ListName == EFriendsLists::ToString(EFriendsLists::InGamePlayers) && !Presence.bIsPlayingThisGame)
-		{
-			continue;
-		}
-		// If the service hasn't returned the info yet, skip them
-		else if (Friend->GetDisplayName().IsEmpty())
+		else if (FriendsListType == EFriendsLists::InGamePlayers && !Presence.bIsPlayingThisGame)
 		{
 			continue;
 		}
diff --git a/Engine/Plugins/Online/OnlineSubsystemEOS/Source/OnlineSubsystemEOS/Private/UserManagerEOS.h b/Engine/Plugins/Online/OnlineSubsystemEOS/Source/OnlineSubsystemEOS/Private/UserManagerEOS.h
--- a/Engine/Plugins/Online/OnlineSubsystemEOS/Source/OnlineSubsystemEOS/Private/UserManagerEOS.h
+++ b/Engine/Plugins/Online/OnlineSubsystemEOS/Source/OnlineSubsystemEOS/Private/UserManagerEOS.h
@@ -395,6 +395,7 @@ public:
 	int32 GetDefaultLocalUser() const { return DefaultLocalUser; }
 
 private:
+	bool IsLocalUserValid(int32 LocalUserNum) const ;
 	FLocalUserEOS& GetLocalUserChecked(int32 LocalUserNum);
 
 	void CallEOSAuthLogin(int32 LocalUserNum, const FOnlineAccountCredentials& Credentials);
