- Add support for GetDisplayName Platform parameter, calling into EOS CopyBestDisplayNameWithPlatform.

diff --git a/Engine/Plugins/Online/EOSShared/Source/EOSShared/Private/EOSShared.cpp b/Engine/Plugins/Online/EOSShared/Source/EOSShared/Private/EOSShared.cpp
--- a/Engine/Plugins/Online/EOSShared/Source/EOSShared/Private/EOSShared.cpp
+++ b/Engine/Plugins/Online/EOSShared/Source/EOSShared/Private/EOSShared.cpp
@@ -495,6 +495,39 @@ bool LexFromString(EOS_EPacketReliability& OutEnum, const TCHAR* InString)
 	return true;
 }
 
+EOS_OnlinePlatformType EOSOnlinePlatformTypeFromString(const FStringView InString)
+{
+	if (InString == TEXT("Unknown"))
+	{
+		return EOS_OPT_Unknown;
+	}
+	else if (InString == TEXT("Epic"))
+	{	
+		return EOS_OPT_Epic;
+	}
+	else if (InString == TEXT("Steam"))
+	{	
+		return EOS_OPT_Steam;
+	}
+	else if (InString == TEXT("PSN"))
+	{	
+		return 1000; //EOS_OPT_PSN;
+	}
+	else if (InString == TEXT("Switch"))
+	{	
+		return 2000; //EOS_OPT_SWITCH;
+	}
+	else if (InString == TEXT("XBL"))
+	{	
+		return 3000; //EOS_OPT_XBL;
+	}
+	else
+	{
+		checkNoEntry();
+		return EOS_OPT_Unknown;
+	}
+}
+
 FString LexToString(const EOS_RTC_Option& Option)
 {
 	UE_EOS_CHECK_API_MISMATCH(EOS_RTC_OPTION_API_LATEST, 1);
diff --git a/Engine/Plugins/Online/EOSShared/Source/EOSShared/Public/EOSShared.h b/Engine/Plugins/Online/EOSShared/Source/EOSShared/Public/EOSShared.h
--- a/Engine/Plugins/Online/EOSShared/Source/EOSShared/Public/EOSShared.h
+++ b/Engine/Plugins/Online/EOSShared/Source/EOSShared/Public/EOSShared.h
@@ -62,6 +62,8 @@ EOS_STRUCT_FORWARD_DECL(EOS_UserInfo_BestDisplayName);
 EOS_STRUCT_FORWARD_DECL(EOS_RTC_Option);
 #undef EOS_STRUCT_FORWARD_DECL
 
+extern "C" typedef uint32_t EOS_OnlinePlatformType;
+
 DECLARE_LOG_CATEGORY_EXTERN(LogEOSSDK, Log, All);
 
 EOSSHARED_API FString LexToString(const EOS_EResult EosResult);
@@ -89,12 +91,13 @@ EOSSHARED_API const TCHAR* LexToString(const EOS_ENetworkStatus NetworkStatus);
 EOSSHARED_API const TCHAR* LexToString(const EOS_Presence_EStatus PresenceStatus);
 
 EOSSHARED_API bool LexFromString(EOS_EAuthScopeFlags& OutEnum, const FStringView InString);
-EOSSHARED_API bool LexFromString(EOS_EExternalAccountType& OutEnum, const TCHAR* InString);
+EOSSHARED_API bool LexFromString(EOS_EIntegratedPlatformManagementFlags& OutEnum, const TCHAR* InString);
 EOSSHARED_API bool LexFromString(EOS_ELoginCredentialType& OutEnum, const TCHAR* InString);
+EOSSHARED_API bool LexFromString(EOS_EPacketReliability& OutEnum, const TCHAR* InString);
 EOSSHARED_API bool LexFromString(EOS_ERTCBackgroundMode& OutEnum, const TCHAR* InString);
 EOSSHARED_API bool LexFromString(EOS_UI_EInputStateButtonFlags& OutEnum, const TCHAR* InString);
-EOSSHARED_API bool LexFromString(EOS_EIntegratedPlatformManagementFlags& OutEnum, const TCHAR* InString);
-EOSSHARED_API bool LexFromString(EOS_EPacketReliability& OutEnum, const TCHAR* InString);
+
+EOSSHARED_API EOS_OnlinePlatformType EOSOnlinePlatformTypeFromString(const FStringView InString);
 
 /** Extracts the display name FString from a EOS_UserInfo_BestDisplayName using the following logic: Nickname > DisplayNameSanitized > DisplayName */
 EOSSHARED_API FString GetBestDisplayNameStr(const EOS_UserInfo_BestDisplayName& BestDisplayName);
diff --git a/Engine/Plugins/Online/OnlineSubsystemEOS/Source/OnlineSubsystemEOS/Private/EOSHelpers.cpp b/Engine/Plugins/Online/OnlineSubsystemEOS/Source/OnlineSubsystemEOS/Private/EOSHelpers.cpp
--- a/Engine/Plugins/Online/OnlineSubsystemEOS/Source/OnlineSubsystemEOS/Private/EOSHelpers.cpp
+++ b/Engine/Plugins/Online/OnlineSubsystemEOS/Source/OnlineSubsystemEOS/Private/EOSHelpers.cpp
@@ -18,7 +18,7 @@ void FEOSHelpers::PlatformTriggerLoginUI(FOnlineSubsystemEOS* EOSSubsystem, cons
 {
 	check(EOSSubsystem);
 
-	UE_LOG_ONLINE_EXTERNALUI(Warning, TEXT("[FUserManagerEOS::ShowLoginUI] This method is not implemented."));
+	UE_LOG_ONLINE_EXTERNALUI(Warning, TEXT("[%hs] This method is not implemented."), __FUNCTION__);
 
 	const FUniqueNetIdPtr PlayerId = EOSSubsystem->UserManager->GetUniquePlayerId(ControllerIndex);
 
diff --git a/Engine/Plugins/Online/OnlineSubsystemEOS/Source/OnlineSubsystemEOS/Private/OnlineSubsystemEOSTypes.h b/Engine/Plugins/Online/OnlineSubsystemEOS/Source/OnlineSubsystemEOS/Private/OnlineSubsystemEOSTypes.h
--- a/Engine/Plugins/Online/OnlineSubsystemEOS/Source/OnlineSubsystemEOS/Private/OnlineSubsystemEOSTypes.h
+++ b/Engine/Plugins/Online/OnlineSubsystemEOS/Source/OnlineSubsystemEOS/Private/OnlineSubsystemEOSTypes.h
@@ -168,14 +168,16 @@ class TOnlineUserEOS
 public:
 	friend class FUserManagerEOS;
 
-	TOnlineUserEOS(const FUniqueNetIdEOSRef& InNetIdRef)
+	TOnlineUserEOS(const FUniqueNetIdEOSRef& InNetIdRef, const FOnlineSubsystemEOS& InSubsystem)
 		: UserIdRef(InNetIdRef)
+		, EOSSubsystem(InSubsystem)
 	{
 	}
 
-	TOnlineUserEOS(const FUniqueNetIdEOSRef& InNetIdRef, const TMap<FString, FString>& InUserAttributes)
+	TOnlineUserEOS(const FUniqueNetIdEOSRef& InNetIdRef, const TMap<FString, FString>& InUserAttributes, const FOnlineSubsystemEOS& InSubsystem)
 		: UserIdRef(InNetIdRef)
 		, UserAttributes(InUserAttributes)
+		, EOSSubsystem(InSubsystem)
 	{
 	}
 
@@ -196,9 +198,7 @@ public:
 
 	virtual FString GetDisplayName(const FString& Platform = FString()) const override
 	{
-		FString ReturnValue;
-		GetUserAttribute(USER_ATTR_DISPLAY_NAME, ReturnValue);
-		return ReturnValue;
+		return EOSSubsystem.UserManager->GetBestDisplayName(UserIdRef->GetEpicAccountId(), Platform);
 	}
 
 	virtual bool GetUserAttribute(const FString& AttrName, FString& OutAttrValue) const override
@@ -237,6 +237,8 @@ protected:
 	FUniqueNetIdEOSRef UserIdRef;
 	/** Additional key/value pair data related to user attribution */
 	TMap<FString, FString> UserAttributes;
+
+	const FOnlineSubsystemEOS& EOSSubsystem;
 };
 
 /**
@@ -246,10 +248,11 @@ template<class BaseClass>
 class TUserOnlineAccountEOS :
 	public TOnlineUserEOS<BaseClass, IAttributeAccessInterface>
 {
+	using Super = TOnlineUserEOS<BaseClass, IAttributeAccessInterface>;
+
 public:
 	TUserOnlineAccountEOS(const FUniqueNetIdEOSRef& InNetIdRef, const FOnlineSubsystemEOS& InSubsystem)
-		: TOnlineUserEOS<BaseClass, IAttributeAccessInterface>(InNetIdRef)
-		, EOSSubsystem(InSubsystem)
+		: Super(InNetIdRef, InSubsystem)
 	{
 	}
 
@@ -258,10 +261,10 @@ public:
 	{
 		FString Token;
 
-		if (const IOnlineIdentityPtr Identity = EOSSubsystem.GetIdentityInterface())
+		if (const IOnlineIdentityPtr Identity = Super::EOSSubsystem.GetIdentityInterface())
 		{
 			const int32 LocalUserNum = Identity->GetLocalUserNumFromPlatformUserId(Identity->GetPlatformUserIdFromUniqueNetId(*this->UserIdRef));
-			Token = EOSSubsystem.GetIdentityInterface()->GetAuthToken(LocalUserNum);
+			Token = Identity->GetAuthToken(LocalUserNum);
 		}
 
 		return Token;
@@ -298,7 +301,6 @@ public:
 protected:
 	/** Additional key/value pair data related to auth */
 	TMap<FString, FString> AdditionalAuthData;
-	const FOnlineSubsystemEOS& EOSSubsystem;
 };
 
 typedef TSharedRef<FOnlineUserPresence> FOnlineUserPresenceRef;
@@ -311,13 +313,13 @@ class TOnlineFriendEOS :
 	public TOnlineUserEOS<BaseClass, IAttributeAccessInterface>
 {
 public:
-	TOnlineFriendEOS(const FUniqueNetIdEOSRef& InNetIdRef)
-		: TOnlineUserEOS<BaseClass, IAttributeAccessInterface>(InNetIdRef)
+	TOnlineFriendEOS(const FUniqueNetIdEOSRef& InNetIdRef, const FOnlineSubsystemEOS& InSubsystem)
+		: TOnlineUserEOS<BaseClass, IAttributeAccessInterface>(InNetIdRef, InSubsystem)
 	{
 	}
 
-	TOnlineFriendEOS(const FUniqueNetIdEOSRef& InNetIdRef, const TMap<FString, FString>& InUserAttributes)
-		: TOnlineUserEOS<BaseClass, IAttributeAccessInterface>(InNetIdRef, InUserAttributes)
+	TOnlineFriendEOS(const FUniqueNetIdEOSRef& InNetIdRef, const TMap<FString, FString>& InUserAttributes, const FOnlineSubsystemEOS& InSubsystem)
+		: TOnlineUserEOS<BaseClass, IAttributeAccessInterface>(InNetIdRef, InUserAttributes, InSubsystem)
 	{
 	}
 
diff --git a/Engine/Plugins/Online/OnlineSubsystemEOS/Source/OnlineSubsystemEOS/Private/UserManagerEOS.cpp b/Engine/Plugins/Online/OnlineSubsystemEOS/Source/OnlineSubsystemEOS/Private/UserManagerEOS.cpp
--- a/Engine/Plugins/Online/OnlineSubsystemEOS/Source/OnlineSubsystemEOS/Private/UserManagerEOS.cpp
+++ b/Engine/Plugins/Online/OnlineSubsystemEOS/Source/OnlineSubsystemEOS/Private/UserManagerEOS.cpp
@@ -1293,7 +1293,7 @@ FLocalUserEOS& FUserManagerEOS::AddLocalUser(int32 LocalUserNum, EOS_EpicAccount
 		NotificationPairPtr->NotificationId = EOS_Connect_AddNotifyAuthExpiration(EOSSubsystem->ConnectHandle, &Options, CallbackObj, CallbackObj->GetCallbackPtr());
 	}
 
-	const FOnlineUserEOSRef UserRef(new FOnlineUserEOS(UserNetId));
+	const FOnlineUserEOSRef UserRef(new FOnlineUserEOS(UserNetId, *EOSSubsystem));
 	UniqueNetIdToUserRefMap.Emplace(UserNetId, UserRef);
 
 	// Once all the fields are set, we start gathering additional information
@@ -1305,55 +1305,69 @@ FLocalUserEOS& FUserManagerEOS::AddLocalUser(int32 LocalUserNum, EOS_EpicAccount
 	return LocalUser;
 }
 
-void FUserManagerEOS::UpdateUserInfo(IAttributeAccessInterfaceRef AttributeAccessRef, EOS_EpicAccountId LocalId, EOS_EpicAccountId AccountId)
+FString FUserManagerEOS::GetBestDisplayName(EOS_EpicAccountId TargetUserId, const FStringView RequestedPlatform) const
 {
-	EOS_UserInfo_CopyUserInfoOptions Options = { };
-	Options.ApiVersion = 3;
-	UE_EOS_CHECK_API_MISMATCH(EOS_USERINFO_COPYUSERINFO_API_LATEST, 3);
-	Options.LocalUserId = LocalId;
-	Options.TargetUserId = AccountId;
+	FString Result;
 
-	EOS_UserInfo* UserInfo = nullptr;
+	EOS_UserInfo_BestDisplayName* EosBestDisplayName = nullptr;
+	EOS_EResult BestDisplayNameResult = EOS_EResult::EOS_Success;
 
-	EOS_EResult CopyResult = EOS_UserInfo_CopyUserInfo(EOSSubsystem->UserInfoHandle, &Options, &UserInfo);
-	if (CopyResult == EOS_EResult::EOS_Success)
-	{
-		AttributeAccessRef->SetInternalAttribute(USER_ATTR_COUNTRY, UTF8_TO_TCHAR(UserInfo->Country));
-		AttributeAccessRef->SetInternalAttribute(USER_ATTR_LANG, UTF8_TO_TCHAR(UserInfo->PreferredLanguage));
-		EOS_UserInfo_Release(UserInfo);
-	}
+	const EOS_EpicAccountId LocalUserId = GetLocalEpicAccountId(GetDefaultLocalUser());
 
-	EOS_UserInfo_CopyBestDisplayNameOptions BestDisplayNameOptions = { };
-	BestDisplayNameOptions.ApiVersion = 1; 
-	UE_EOS_CHECK_API_MISMATCH(EOS_USERINFO_COPYBESTDISPLAYNAME_API_LATEST, 1);
-	BestDisplayNameOptions.LocalUserId = LocalId;
-	BestDisplayNameOptions.TargetUserId = AccountId;
+	if (RequestedPlatform.IsEmpty())
+	{
+		EOS_UserInfo_CopyBestDisplayNameOptions BestDisplayNameOptions = { };
+		BestDisplayNameOptions.ApiVersion = 1; 
+		UE_EOS_CHECK_API_MISMATCH(EOS_USERINFO_COPYBESTDISPLAYNAME_API_LATEST, 1);
+		BestDisplayNameOptions.LocalUserId = LocalUserId;
+		BestDisplayNameOptions.TargetUserId = TargetUserId;
 
-	EOS_UserInfo_BestDisplayName* BestDisplayName;
-	EOS_EResult BestDisplayNameResult = EOS_UserInfo_CopyBestDisplayName(EOSSubsystem->UserInfoHandle, &BestDisplayNameOptions, &BestDisplayName);
+		BestDisplayNameResult = EOS_UserInfo_CopyBestDisplayName(EOSSubsystem->UserInfoHandle, &BestDisplayNameOptions, &EosBestDisplayName);
+	}
 
-	if (BestDisplayNameResult == EOS_EResult::EOS_UserInfo_BestDisplayNameIndeterminate)
+	if (!RequestedPlatform.IsEmpty() || BestDisplayNameResult == EOS_EResult::EOS_UserInfo_BestDisplayNameIndeterminate)
 	{
 		EOS_UserInfo_CopyBestDisplayNameWithPlatformOptions BestDisplayNameWithPlatformOptions = {};
 		BestDisplayNameWithPlatformOptions.ApiVersion = 1;
 		UE_EOS_CHECK_API_MISMATCH(EOS_USERINFO_COPYBESTDISPLAYNAMEWITHPLATFORM_API_LATEST, 1);
-		BestDisplayNameWithPlatformOptions.LocalUserId = LocalId;
-		BestDisplayNameWithPlatformOptions.TargetUserId = AccountId;
+		BestDisplayNameWithPlatformOptions.LocalUserId = LocalUserId;
+		BestDisplayNameWithPlatformOptions.TargetUserId = TargetUserId;
+		
 		BestDisplayNameWithPlatformOptions.TargetPlatformType = EOS_OPT_Epic;
+		if (!RequestedPlatform.IsEmpty())
+		{
+			BestDisplayNameWithPlatformOptions.TargetPlatformType = EOSOnlinePlatformTypeFromString(RequestedPlatform);
+		}
 
-		BestDisplayNameResult = EOS_UserInfo_CopyBestDisplayNameWithPlatform(EOSSubsystem->UserInfoHandle, &BestDisplayNameWithPlatformOptions, &BestDisplayName);
+		BestDisplayNameResult = EOS_UserInfo_CopyBestDisplayNameWithPlatform(EOSSubsystem->UserInfoHandle, &BestDisplayNameWithPlatformOptions, &EosBestDisplayName);
 	}
 
-	if (BestDisplayNameResult == EOS_EResult::EOS_Success)
+	if (EosBestDisplayName)
 	{
-		// We'll prioritize which name is chosen: Nickname > DisplayNameSanitized > DisplayName
-		AttributeAccessRef->SetInternalAttribute(USER_ATTR_DISPLAY_NAME, *GetBestDisplayNameStr(*BestDisplayName));
-
-		EOS_UserInfo_BestDisplayName_Release(BestDisplayName);
+		Result = GetBestDisplayNameStr(*EosBestDisplayName);
+		EOS_UserInfo_BestDisplayName_Release(EosBestDisplayName);
 	}
-	else
+
+	return Result;
+}
+
+void FUserManagerEOS::UpdateUserInfo(IAttributeAccessInterfaceRef AttributeAccessRef, EOS_EpicAccountId LocalId, EOS_EpicAccountId AccountId)
+{
+	EOS_UserInfo_CopyUserInfoOptions Options = { };
+	Options.ApiVersion = 3;
+	UE_EOS_CHECK_API_MISMATCH(EOS_USERINFO_COPYUSERINFO_API_LATEST, 3);
+	Options.LocalUserId = LocalId;
+	Options.TargetUserId = AccountId;
+
+	EOS_UserInfo* UserInfo = nullptr;
+
+	const EOS_EResult CopyResult = EOS_UserInfo_CopyUserInfo(EOSSubsystem->UserInfoHandle, &Options, &UserInfo);
+	UE_CLOG_ONLINE(CopyResult != EOS_EResult::EOS_Success, Warning, TEXT("%hs Result=[%s]"), __FUNCTION__, *LexToString(CopyResult));
+	if (CopyResult == EOS_EResult::EOS_Success)
 	{
-		UE_LOG_ONLINE(Warning, TEXT("[FUserManagerEOS::UpdateUserInfo] Failed to retrieve BestDisplayName EOS_EResult: %s."), *LexToString(BestDisplayNameResult));
+		AttributeAccessRef->SetInternalAttribute(USER_ATTR_COUNTRY, UTF8_TO_TCHAR(UserInfo->Country));
+		AttributeAccessRef->SetInternalAttribute(USER_ATTR_LANG, UTF8_TO_TCHAR(UserInfo->PreferredLanguage));
+		EOS_UserInfo_Release(UserInfo);
 	}
 }
 
@@ -2191,7 +2205,7 @@ FOnlineFriendEOSRef FUserManagerEOS::AddFriend(int32 LocalUserNum, const FUnique
 {
 	const FOnlineUserEOSRef UserRef = UniqueNetIdToUserRefMap[FriendNetId.AsShared()];
 	const FUniqueNetIdEOSRef FriendNetIdEOSRef = StaticCastSharedRef<const FUniqueNetIdEOS>(FriendNetId.AsShared());
-	const FOnlineFriendEOSRef FriendRef = MakeShareable(new FOnlineFriendEOS(FriendNetIdEOSRef, UserRef->UserAttributes));
+	const FOnlineFriendEOSRef FriendRef = MakeShareable(new FOnlineFriendEOS(FriendNetIdEOSRef, UserRef->UserAttributes, *EOSSubsystem));
 
 	GetLocalUserChecked(LocalUserNum).FriendsList->Add(FriendNetId.AsShared(), FriendRef);
 
@@ -2219,7 +2233,7 @@ void FUserManagerEOS::AddRemotePlayer(int32 LocalUserNum, EOS_EpicAccountId Epic
 {
 	const FResolveUniqueNetIdCallback IdResolutionCallback = [this, LocalUserNum, EpicAccountId, Callback](FUniqueNetIdEOSRef ResolvedUniqueNetId) mutable
 	{
-		const FOnlineUserEOSRef AttributeRef = MakeShareable(new FOnlineUserEOS(ResolvedUniqueNetId));
+		const FOnlineUserEOSRef AttributeRef = MakeShareable(new FOnlineUserEOS(ResolvedUniqueNetId, *EOSSubsystem));
 
 		UniqueNetIdToUserRefMap.Emplace(ResolvedUniqueNetId, AttributeRef);
 
@@ -2236,7 +2250,7 @@ void FUserManagerEOS::AddRemotePlayers(int32 LocalUserNum, TArray<EOS_EpicAccoun
 	{
 		for (const TPair<EOS_EpicAccountId, FUniqueNetIdEOSRef>& Entry : ResolvedUniqueNetIds)
 		{
-			const FOnlineUserEOSRef AttributeRef = MakeShareable(new FOnlineUserEOS(Entry.Value));
+			const FOnlineUserEOSRef AttributeRef = MakeShareable(new FOnlineUserEOS(Entry.Value, *EOSSubsystem));
 
 			UniqueNetIdToUserRefMap.Emplace(Entry.Value, AttributeRef);
 
diff --git a/Engine/Plugins/Online/OnlineSubsystemEOS/Source/OnlineSubsystemEOS/Private/UserManagerEOS.h b/Engine/Plugins/Online/OnlineSubsystemEOS/Source/OnlineSubsystemEOS/Private/UserManagerEOS.h
--- a/Engine/Plugins/Online/OnlineSubsystemEOS/Source/OnlineSubsystemEOS/Private/UserManagerEOS.h
+++ b/Engine/Plugins/Online/OnlineSubsystemEOS/Source/OnlineSubsystemEOS/Private/UserManagerEOS.h
@@ -30,8 +30,8 @@ class FOnlineUserEOS :
 	public TOnlineUserEOS<FOnlineUser, IAttributeAccessInterface>
 {
 public:
-	FOnlineUserEOS(FUniqueNetIdEOSRef InUserId) :
-		TOnlineUserEOS<FOnlineUser, IAttributeAccessInterface>(InUserId)
+	FOnlineUserEOS(FUniqueNetIdEOSRef InUserId, const FOnlineSubsystemEOS& InSubsystem) :
+		TOnlineUserEOS<FOnlineUser, IAttributeAccessInterface>(InUserId, InSubsystem)
 	{
 	}
 	virtual ~FOnlineUserEOS() = default;
@@ -64,13 +64,13 @@ class FOnlineFriendEOS :
 	public TOnlineFriendEOS<FOnlineFriend>
 {
 public:
-	FOnlineFriendEOS(FUniqueNetIdEOSRef InUserId) :
-		TOnlineFriendEOS<FOnlineFriend>(InUserId)
+	FOnlineFriendEOS(FUniqueNetIdEOSRef InUserId, const FOnlineSubsystemEOS& InSubsystem) :
+		TOnlineFriendEOS<FOnlineFriend>(InUserId, InSubsystem)
 	{
 	}
 
-	FOnlineFriendEOS(FUniqueNetIdEOSRef InUserId, const TMap<FString, FString>& InUserAttributes) :
-		TOnlineFriendEOS<FOnlineFriend>(InUserId, InUserAttributes)
+	FOnlineFriendEOS(FUniqueNetIdEOSRef InUserId, const TMap<FString, FString>& InUserAttributes, const FOnlineSubsystemEOS& InSubsystem) :
+		TOnlineFriendEOS<FOnlineFriend>(InUserId, InUserAttributes, InSubsystem)
 	{
 	}
 
@@ -394,6 +394,8 @@ public:
 
 	int32 GetDefaultLocalUser() const { return DefaultLocalUser; }
 
+	FString GetBestDisplayName(EOS_EpicAccountId TargetUserId, const FStringView RequestedPlatform) const;
+
 private:
 	bool IsLocalUserValid(int32 LocalUserNum) const ;
 	FLocalUserEOS& GetLocalUserChecked(int32 LocalUserNum);
