Added a IsValid() check on CurrentSessionSearch in StartLobbySearch to avoid potential crash

diff --git a/Engine/Plugins/Online/OnlineSubsystemEOS/Source/OnlineSubsystemEOS/Private/OnlineSessionEOS.cpp b/Engine/Plugins/Online/OnlineSubsystemEOS/Source/OnlineSubsystemEOS/Private/OnlineSessionEOS.cpp
--- a/Engine/Plugins/Online/OnlineSubsystemEOS/Source/OnlineSubsystemEOS/Private/OnlineSessionEOS.cpp
+++ b/Engine/Plugins/Online/OnlineSubsystemEOS/Source/OnlineSubsystemEOS/Private/OnlineSessionEOS.cpp
@@ -4577,7 +4577,13 @@ void FOnlineSessionEOS::StartLobbySearch(int32 SearchingPlayerNum, EOS_HLobbySea
 	LobbySearchFindCallback = CallbackObj;
 	CallbackObj->CallbackLambda = [this, SearchingPlayerNum, LobbySearchHandle, SearchSettings, CompletionDelegate](const EOS_LobbySearch_FindCallbackInfo* Data)
 	{
-		if (Data->ResultCode == EOS_EResult::EOS_Success)
+		if (!CurrentSessionSearch.IsValid())
+		{
+			UE_LOG_ONLINE_SESSION(Log, TEXT("[FOnlineSessionEOS::StartLobbySearch] Current session search is invalid. It may have been canceled."));
+			CompletionDelegate.ExecuteIfBound(SearchingPlayerNum, false, FOnlineSessionSearchResult());
+		}
+		
+		else if (Data->ResultCode == EOS_EResult::EOS_Success)
 		{
 			UE_LOG_ONLINE_SESSION(Log, TEXT("[FOnlineSessionEOS::StartLobbySearch] LobbySearch_Find was successful."));
 
