Added logic to track if a friends query has been processed for a local user in UserManagerEOS, and have GetFriendsList return false if it hasn't.
Fixed bug with id resolution method not calling its callback function if passed an empty list.

diff --git a/Engine/Plugins/Online/OnlineSubsystemEOS/Source/OnlineSubsystemEOS/Private/UserManagerEOS.cpp b/Engine/Plugins/Online/OnlineSubsystemEOS/Source/OnlineSubsystemEOS/Private/UserManagerEOS.cpp
--- a/Engine/Plugins/Online/OnlineSubsystemEOS/Source/OnlineSubsystemEOS/Private/UserManagerEOS.cpp
+++ b/Engine/Plugins/Online/OnlineSubsystemEOS/Source/OnlineSubsystemEOS/Private/UserManagerEOS.cpp
@@ -1270,8 +1270,6 @@ FLocalUserEOS& FUserManagerEOS::AddLocalUser(int32 LocalUserNum, EOS_EpicAccount
 	FUserOnlineAccountEOSRef UserAccountRef(new FUserOnlineAccountEOS(UserNetId, *EOSSubsystem));
 	LocalUser.UserOnlineAccount = UserAccountRef;
 
-	LocalUser.FriendsList = MakeShareable(new FFriendsListEOS(LocalUserNum, UserNetId));
-
 	LocalUser.LastLoginCredentials = MakeShared<FOnlineAccountCredentials>(AccountCredentials);
 
 	// Add auth refresh notification if not set for this user yet
@@ -2137,25 +2135,32 @@ void FUserManagerEOS::FriendStatusChangedImpl(EOS_EpicAccountId LocalUserId, EOS
 	case EOS_EFriendsStatus::EOS_FS_NotFriends: // Invite rejections and friend removal
 	{
 		//User should already be a friend
-		FFriendsListEOSPtr& FriendsListRef = GetLocalUserChecked(LocalUserNum).FriendsList;
-		FOnlineFriendEOSPtr Friend = FriendsListRef->GetByNetId(FriendEOSId);
-		if (Friend.IsValid())
+		FFriendsListEOSPtr& FriendsListPtr = GetLocalUserChecked(LocalUserNum).FriendsList;
+		if (FriendsListPtr)
 		{
-			FriendsListRef->Remove(FriendEOSId, Friend.ToSharedRef());
-			Friend->SetInviteStatus(EInviteStatus::Unknown);
-
-			if (PreviousStatus == EOS_EFriendsStatus::EOS_FS_Friends)
+			FOnlineFriendEOSPtr Friend = FriendsListPtr->GetByNetId(FriendEOSId);
+			if (Friend.IsValid())
 			{
-				TriggerOnFriendRemovedDelegates(*LocalEOSID, *FriendEOSId);
+				FriendsListPtr->Remove(FriendEOSId, Friend.ToSharedRef());
+				Friend->SetInviteStatus(EInviteStatus::Unknown);
+
+				if (PreviousStatus == EOS_EFriendsStatus::EOS_FS_Friends)
+				{
+					TriggerOnFriendRemovedDelegates(*LocalEOSID, *FriendEOSId);
+				}
+				else if (PreviousStatus == EOS_EFriendsStatus::EOS_FS_InviteSent || PreviousStatus == EOS_EFriendsStatus::EOS_FS_InviteReceived)
+				{
+					TriggerOnInviteRejectedDelegates(*LocalEOSID, *FriendEOSId); // We don't have an "OnInviteRejected" event only for the local user
+				}
 			}
-			else if (PreviousStatus == EOS_EFriendsStatus::EOS_FS_InviteSent || PreviousStatus == EOS_EFriendsStatus::EOS_FS_InviteReceived)
+			else
 			{
-				TriggerOnInviteRejectedDelegates(*LocalEOSID, *FriendEOSId); // We don't have an "OnInviteRejected" event only for the local user
+				UE_LOG_ONLINE_FRIEND(Verbose, TEXT("Friend status notification received for user [%d], but remote user [%s] was not previously registered as a friend"), LocalUserNum, *FriendEOSId->ToString());
 			}
 		}
 		else
 		{
-			UE_LOG_ONLINE_FRIEND(Verbose, TEXT("Friend status notification received for user [%d], but remote user [%s] was not previously registered as a friend"), *FriendEOSId->ToString());
+			UE_LOG_ONLINE_FRIEND(Verbose, TEXT("Friend list still has not been queried for local user [%d]"), LocalUserNum);
 		}
 
 		break;
@@ -2203,11 +2208,15 @@ void FUserManagerEOS::FriendStatusChangedImpl(EOS_EpicAccountId LocalUserId, EOS
 
 FOnlineFriendEOSRef FUserManagerEOS::AddFriend(int32 LocalUserNum, const FUniqueNetIdEOS& FriendNetId)
 {
+	// A call to AddFriend should only be made after the friends list has been initialised
+	FFriendsListEOSPtr FriendsListPtr = GetLocalUserChecked(LocalUserNum).FriendsList;
+	check(FriendsListPtr);
+
 	const FOnlineUserEOSRef UserRef = UniqueNetIdToUserRefMap[FriendNetId.AsShared()];
 	const FUniqueNetIdEOSRef FriendNetIdEOSRef = StaticCastSharedRef<const FUniqueNetIdEOS>(FriendNetId.AsShared());
 	const FOnlineFriendEOSRef FriendRef = MakeShareable(new FOnlineFriendEOS(FriendNetIdEOSRef, UserRef->UserAttributes, *EOSSubsystem));
 
-	GetLocalUserChecked(LocalUserNum).FriendsList->Add(FriendNetId.AsShared(), FriendRef);
+	FriendsListPtr->Add(FriendNetId.AsShared(), FriendRef);
 
 	EOS_Friends_GetStatusOptions Options = { };
 	Options.ApiVersion = 1;
@@ -2248,37 +2257,46 @@ void FUserManagerEOS::AddRemotePlayers(int32 LocalUserNum, TArray<EOS_EpicAccoun
 {
 	const FResolveEpicAccountIdsCallback IdResolutionCallback = [this, LocalUserNum, EpicAccountIds, Callback](bool bWasSuccessful, TMap<EOS_EpicAccountId, FUniqueNetIdEOSRef> ResolvedUniqueNetIds, FString ErrorStr) mutable
 	{
-		for (const TPair<EOS_EpicAccountId, FUniqueNetIdEOSRef>& Entry : ResolvedUniqueNetIds)
+		if (!ResolvedUniqueNetIds.IsEmpty())
 		{
-			const FOnlineUserEOSRef AttributeRef = MakeShareable(new FOnlineUserEOS(Entry.Value, *EOSSubsystem));
+			for (const TPair<EOS_EpicAccountId, FUniqueNetIdEOSRef>& Entry : ResolvedUniqueNetIds)
+			{
+				const FOnlineUserEOSRef AttributeRef = MakeShareable(new FOnlineUserEOS(Entry.Value, *EOSSubsystem));
 
-			UniqueNetIdToUserRefMap.Emplace(Entry.Value, AttributeRef);
+				UniqueNetIdToUserRefMap.Emplace(Entry.Value, AttributeRef);
 
-			const FRemoteUserProcessedCallback ReadUserInfoCallback = [this, LocalUserNum, Callback](bool bWasSuccessful, FUniqueNetIdEOSRef UserNetId, const FString& ErrorStr) mutable
-			{
-				FLocalUserEOS& LocalUser = GetLocalUserChecked(LocalUserNum);
+				const FRemoteUserProcessedCallback ReadUserInfoCallback = [this, LocalUserNum, Callback](bool bWasSuccessful, FUniqueNetIdEOSRef UserNetId, const FString& ErrorStr) mutable
+					{
+						FLocalUserEOS& LocalUser = GetLocalUserChecked(LocalUserNum);
 
-				LocalUser.OngoingQueryUserInfoResults.ProcessedIds.Add(UserNetId);
+						LocalUser.OngoingQueryUserInfoResults.ProcessedIds.Add(UserNetId);
 
-				LocalUser.OngoingQueryUserInfoResults.bAllWasSuccessful &= bWasSuccessful;
+						LocalUser.OngoingQueryUserInfoResults.bAllWasSuccessful &= bWasSuccessful;
 
-				if (!ErrorStr.IsEmpty())
-				{
-					LocalUser.OngoingQueryUserInfoResults.AllErrorStr += TEXT("/n") + ErrorStr;
-				}
+						if (!ErrorStr.IsEmpty())
+						{
+							LocalUser.OngoingQueryUserInfoResults.AllErrorStr += TEXT("/n") + ErrorStr;
+						}
 
-				if (LocalUser.OngoingQueryUserInfoAccounts.IsEmpty())
-				{
-					FLocalUserEOS::FReadUserInfoResults Results = LocalUser.OngoingQueryUserInfoResults;
-					LocalUser.OngoingQueryUserInfoResults.Reset();
-					
-					Callback(Results.bAllWasSuccessful, Results.ProcessedIds, Results.AllErrorStr);
-				}
-			};
+						if (LocalUser.OngoingQueryUserInfoAccounts.IsEmpty())
+						{
+							FLocalUserEOS::FReadUserInfoResults Results = LocalUser.OngoingQueryUserInfoResults;
+							LocalUser.OngoingQueryUserInfoResults.Reset();
+
+							Callback(Results.bAllWasSuccessful, Results.ProcessedIds, Results.AllErrorStr);
+						}
+					};
 
-			// Read the user info for this player
-			ReadUserInfo(LocalUserNum, Entry.Key, ReadUserInfoCallback);
-		}		
+				// Read the user info for this player
+				ReadUserInfo(LocalUserNum, Entry.Key, ReadUserInfoCallback);
+			}
+		}
+		else
+		{
+			// If the list is empty, we still need to execute the callback
+			TArray<FUniqueNetIdEOSRef> EmptyIdArray;
+			Callback(bWasSuccessful, EmptyIdArray, ErrorStr);
+		}			
 	};
 
 	ResolveUniqueNetIds(LocalUserNum, EpicAccountIds, IdResolutionCallback);
@@ -2352,7 +2370,15 @@ bool FUserManagerEOS::ReadFriendsList(int32 LocalUserNum, const FString& ListNam
 			{
 				if (bWasSuccessful)
 				{
-					GetLocalUserChecked(LocalUserNum).FriendsList->Empty(FriendCount);
+					FFriendsListEOSPtr FriendsListPtr = GetLocalUserChecked(LocalUserNum).FriendsList;
+					if (FriendsListPtr)
+					{
+						FriendsListPtr->Empty(FriendCount);
+					}
+					else
+					{
+						GetLocalUserChecked(LocalUserNum).FriendsList = MakeShareable(new FFriendsListEOS(LocalUserNum, GetLocalUniqueNetIdEOS(LocalUserNum).ToSharedRef()));
+					}
 
 					for (const FUniqueNetIdEOSRef& RemoteUserNetId : RemoteUserNetIds)
 					{
@@ -2592,9 +2618,15 @@ bool FUserManagerEOS::GetFriendsList(int32 LocalUserNum, const FString& ListName
 		return false;
 	}
 
+	const FFriendsListEOSPtr FriendsListPtr = GetLocalUserChecked(LocalUserNum).FriendsList;
+	if (FriendsListPtr == nullptr)
+	{
+		return false;
+	}
+
 	OutFriends.Reset();
 
-	for (FOnlineFriendEOSRef Friend : GetLocalUserChecked(LocalUserNum).FriendsList->GetList())
+	for (FOnlineFriendEOSRef Friend : FriendsListPtr->GetList())
 	{
 		const FOnlineUserPresence& Presence = Friend->GetPresence();
 		// See if they only want online only
@@ -2640,23 +2672,27 @@ bool FUserManagerEOS::GetFriendsList(int32 LocalUserNum, const FString& ListName
 
 TSharedPtr<FOnlineFriend> FUserManagerEOS::GetFriend(int32 LocalUserNum, const FUniqueNetId& FriendId, const FString& ListName)
 {
-	const FUniqueNetIdEOS& EosId = FUniqueNetIdEOS::Cast(FriendId);
-	FOnlineFriendEOSPtr FoundFriend = GetLocalUserChecked(LocalUserNum).FriendsList->GetByNetId(EosId.AsShared());
-	if (FoundFriend.IsValid())
+	FFriendsListEOSPtr FriendsListPtr = GetLocalUserChecked(LocalUserNum).FriendsList;
+	if (FriendsListPtr)
 	{
-		const FOnlineUserPresence& Presence = FoundFriend->GetPresence();
-		// See if they only want online only
-		if (ListName == EFriendsLists::ToString(EFriendsLists::OnlinePlayers) && !Presence.bIsOnline)
+		const FUniqueNetIdEOS& EosId = FUniqueNetIdEOS::Cast(FriendId);
+		FOnlineFriendEOSPtr FoundFriend = FriendsListPtr->GetByNetId(EosId.AsShared());
+		if (FoundFriend.IsValid())
 		{
-			return TSharedPtr<FOnlineFriend>();
-		}
-		// Of if they only want friends playing this game
-		else if (ListName == EFriendsLists::ToString(EFriendsLists::InGamePlayers) && !Presence.bIsPlayingThisGame)
-		{
-			return TSharedPtr<FOnlineFriend>();
-		}
+			const FOnlineUserPresence& Presence = FoundFriend->GetPresence();
+			// See if they only want online only
+			if (ListName == EFriendsLists::ToString(EFriendsLists::OnlinePlayers) && !Presence.bIsOnline)
+			{
+				return TSharedPtr<FOnlineFriend>();
+			}
+			// Of if they only want friends playing this game
+			else if (ListName == EFriendsLists::ToString(EFriendsLists::InGamePlayers) && !Presence.bIsPlayingThisGame)
+			{
+				return TSharedPtr<FOnlineFriend>();
+			}
 
-		return FoundFriend;
+			return FoundFriend;
+		}
 	}
 
 	return TSharedPtr<FOnlineFriend>();
@@ -3184,12 +3220,6 @@ void FUserManagerEOS::ReadUserInfo(int32 LocalUserNum, EOS_EpicAccountId EpicAcc
 		{
 			IAttributeAccessInterfaceRef AttributeAccessRef = UniqueNetIdToUserRefMap[EOSId];
 			UpdateUserInfo(AttributeAccessRef, Data->LocalUserId, Data->TargetUserId);
-
-			FLocalUserEOS& LocalUser = GetLocalUserChecked(LocalUserNum);
-			if (FOnlineFriendEOSPtr FriendPtr = LocalUser.FriendsList->GetByNetId(EOSId))
-			{
-				FriendPtr->UpdateInternalAttributes(AttributeAccessRef->GetInternalAttributes());
-			}
 		}
 		else
 		{
diff --git a/Engine/Plugins/Online/OnlineSubsystemEOS/Source/OnlineSubsystemEOS/Private/UserManagerEOS.h b/Engine/Plugins/Online/OnlineSubsystemEOS/Source/OnlineSubsystemEOS/Private/UserManagerEOS.h
--- a/Engine/Plugins/Online/OnlineSubsystemEOS/Source/OnlineSubsystemEOS/Private/UserManagerEOS.h
+++ b/Engine/Plugins/Online/OnlineSubsystemEOS/Source/OnlineSubsystemEOS/Private/UserManagerEOS.h
@@ -200,7 +200,7 @@ struct FLocalUserEOS
 	/** Online Account information for the local user */
 	FUserOnlineAccountEOSPtr UserOnlineAccount;
 
-	/** Friends information for the local user */
+	/** Friends information for the local user. Populated by ReadFriendsList */
 	FFriendsListEOSPtr FriendsList;
 
 	/** Epic ids for users with ongoing queries */
