Fix crash on login failure.

diff --git a/Engine/Plugins/Online/OnlineSubsystemEOS/Source/OnlineSubsystemEOS/Private/UserManagerEOS.cpp b/Engine/Plugins/Online/OnlineSubsystemEOS/Source/OnlineSubsystemEOS/Private/UserManagerEOS.cpp
--- a/Engine/Plugins/Online/OnlineSubsystemEOS/Source/OnlineSubsystemEOS/Private/UserManagerEOS.cpp
+++ b/Engine/Plugins/Online/OnlineSubsystemEOS/Source/OnlineSubsystemEOS/Private/UserManagerEOS.cpp
@@ -1790,10 +1790,18 @@ void FUserManagerEOS::RemoveLocalUser(int32 LocalUserNum)
 {
 	if (LocalUsers.IsValidIndex(LocalUserNum))
 	{
-		const FUniqueNetIdEOSRef FoundId = GetLocalUserChecked(LocalUserNum).UniqueNetId.ToSharedRef();
+		FLocalUserEOS& LocalUser = GetLocalUserChecked(LocalUserNum);
+		const FUniqueNetIdEOSRef FoundId = LocalUser.UniqueNetId.ToSharedRef();
 
 		EOSSubsystem->ReleaseVoiceChatUserInterface(*FoundId);
 
+
+		for (const ReadUserListInfo& CachedInfo : LocalUser.CachedReadUserListInfo)
+		{
+			CachedInfo.ExecuteDelegateIfBound(false, TEXT("User has been removed"));
+		}
+		LocalUser.CachedReadUserListInfo.Empty();
+
 		LocalUsers.RemoveAt(LocalUserNum);
 	}
 
@@ -2331,13 +2339,21 @@ bool FUserManagerEOS::ReadFriendsList(int32 LocalUserNum, const FString& ListNam
 	Options.LocalUserId = GetLocalEpicAccountId(LocalUserNum);
 
 	FReadFriendsCallback* CallbackObj = new FReadFriendsCallback(AsWeak());
-	CallbackObj->CallbackLambda = [this, LocalUserNum, ListName, Delegate](const EOS_Friends_QueryFriendsCallbackInfo* Data)
+	CallbackObj->CallbackLambda = [this, LocalUserNum, ListName](const EOS_Friends_QueryFriendsCallbackInfo* Data)
 	{
 		EOS_EResult Result = Data->ResultCode;
-		if (GetLoginStatus(LocalUserNum) != ELoginStatus::LoggedIn)
+		if (FUniqueNetIdEOSPtr UserId = GetLocalUniqueNetIdEOS(LocalUserNum))
 		{
-			// Handle the user logging out while a read is in progress
-			Result = EOS_EResult::EOS_InvalidUser;
+			if (GetLoginStatus(*UserId) != ELoginStatus::LoggedIn)
+			{
+				// Handle the user logging out while a read is in progress
+				Result = EOS_EResult::EOS_InvalidUser;
+			}
+		}
+		else
+		{
+			UE_LOG_ONLINE_FRIEND(Verbose, TEXT("[FUserManagerEOS::ReadFriendsList] User is no longer available. Abandoning callback."));
+			return;
 		}
 
 		bool bWasSuccessful = Result == EOS_EResult::EOS_Success;
