Addition of support for -epicdeploymentid and -deploymentid parameters to define EOS platform creation

diff --git a/Engine/Plugins/Online/OnlineSubsystemEOS/Source/OnlineSubsystemEOS/Private/EOSSettings.cpp b/Engine/Plugins/Online/OnlineSubsystemEOS/Source/OnlineSubsystemEOS/Private/EOSSettings.cpp
--- a/Engine/Plugins/Online/OnlineSubsystemEOS/Source/OnlineSubsystemEOS/Private/EOSSettings.cpp
+++ b/Engine/Plugins/Online/OnlineSubsystemEOS/Source/OnlineSubsystemEOS/Private/EOSSettings.cpp
@@ -246,7 +246,24 @@ bool UEOSSettings::GetSelectedArtifactSettings(FEOSArtifactSettings& OutSettings
 	bool bHasSandboxId = FParse::Value(FCommandLine::Get(), TEXT("EpicSandboxId="), SandboxId);
 	// Prefer -EpicSandboxIdOverride over previous.
 	bHasSandboxId |= FParse::Value(FCommandLine::Get(), TEXT("EpicSandboxIdOverride="), SandboxId);
+
+	FString DeploymentId;
+	// Get the -epicdeploymentid argument. This generally comes from EGS.
+	bool bHasDeploymentId = FParse::Value(FCommandLine::Get(), TEXT("EpicDeploymentId="), DeploymentId);
+	// Prefer -EpicDeploymentIdOverride over previous.
+	bHasDeploymentId |= FParse::Value(FCommandLine::Get(), TEXT("EpicDeploymentIdOverride="), DeploymentId);
+
 	// If present, grab the settings where both match
+	if (bHasSandboxId && bHasDeploymentId)
+	{
+		if (GetArtifactSettings(ArtifactName, SandboxId, DeploymentId, OutSettings))
+		{
+			return true;
+		}
+
+		UE_LOG_ONLINE(Log, TEXT("UEOSSettings::GetSelectedArtifactSettings() ArtifactName=[%s] SandboxId=[%s] DeploymentId=[%s] no settings found for trio, falling back on pair check."), *ArtifactName, *SandboxId, *DeploymentId);
+	}
+
 	if (bHasSandboxId)
 	{
 		if (GetArtifactSettings(ArtifactName, SandboxId, OutSettings))
@@ -277,23 +294,29 @@ FString UEOSSettings::GetDefaultArtifactName()
 
 bool UEOSSettings::GetArtifactSettings(const FString& ArtifactName, FEOSArtifactSettings& OutSettings)
 {
-	return GetArtifactSettingsImpl(ArtifactName, TOptional<FString>(), OutSettings);
+	return GetArtifactSettingsImpl(ArtifactName, TOptional<FString>(), TOptional<FString>(), OutSettings);
 }
 
 bool UEOSSettings::GetArtifactSettings(const FString& ArtifactName, const FString& SandboxId, FEOSArtifactSettings& OutSettings)
 {
-	return GetArtifactSettingsImpl(ArtifactName, SandboxId, OutSettings);
+	return GetArtifactSettingsImpl(ArtifactName, SandboxId, TOptional<FString>(), OutSettings);
+}
+
+bool UEOSSettings::GetArtifactSettings(const FString& ArtifactName, const FString& SandboxId, const FString& DeploymentId, FEOSArtifactSettings& OutSettings)
+{
+	return GetArtifactSettingsImpl(ArtifactName, SandboxId, DeploymentId, OutSettings);
 }
 
-bool UEOSSettings::GetArtifactSettingsImpl(const FString& ArtifactName, const TOptional<FString>& SandboxId, FEOSArtifactSettings& OutSettings)
+bool UEOSSettings::GetArtifactSettingsImpl(const FString& ArtifactName, const TOptional<FString>& SandboxId, const TOptional<FString>& DeploymentId, FEOSArtifactSettings& OutSettings)
 {
 	if (UObjectInitialized())
 	{
 		const UEOSSettings* This = GetDefault<UEOSSettings>();
-		const FArtifactSettings* Found = This->Artifacts.FindByPredicate([&ArtifactName, &SandboxId](const FArtifactSettings& Element)
+		const FArtifactSettings* Found = This->Artifacts.FindByPredicate([&ArtifactName, &SandboxId, &DeploymentId](const FArtifactSettings& Element)
 		{
 			return Element.ArtifactName == ArtifactName
-				&& (!SandboxId.IsSet() || Element.SandboxId == SandboxId);
+				&& (!SandboxId.IsSet() || Element.SandboxId == SandboxId)
+				&& (!DeploymentId.IsSet() || Element.DeploymentId == DeploymentId);
 		});
 		if (Found)
 		{
@@ -305,10 +328,11 @@ bool UEOSSettings::GetArtifactSettingsImpl(const FString& ArtifactName, const TO
 	else
 	{
 		const TArray<FEOSArtifactSettings>& CachedSettings = GetCachedArtifactSettings();
-		const FEOSArtifactSettings* Found = CachedSettings.FindByPredicate([&ArtifactName, &SandboxId](const FEOSArtifactSettings& Element)
+		const FEOSArtifactSettings* Found = CachedSettings.FindByPredicate([&ArtifactName, &SandboxId, &DeploymentId](const FEOSArtifactSettings& Element)
 		{
 			return Element.ArtifactName == ArtifactName
-				&& (!SandboxId.IsSet() || Element.SandboxId == SandboxId);
+				&& (!SandboxId.IsSet() || Element.SandboxId == SandboxId)
+				&& (!DeploymentId.IsSet() || Element.DeploymentId == DeploymentId);
 		});
 		if (Found)
 		{
diff --git a/Engine/Plugins/Online/OnlineSubsystemEOS/Source/OnlineSubsystemEOS/Public/EOSSettings.h b/Engine/Plugins/Online/OnlineSubsystemEOS/Source/OnlineSubsystemEOS/Public/EOSSettings.h
--- a/Engine/Plugins/Online/OnlineSubsystemEOS/Source/OnlineSubsystemEOS/Public/EOSSettings.h
+++ b/Engine/Plugins/Online/OnlineSubsystemEOS/Source/OnlineSubsystemEOS/Public/EOSSettings.h
@@ -201,7 +201,8 @@ private:
 
 	static bool GetArtifactSettings(const FString& ArtifactName, FEOSArtifactSettings& OutSettings);
 	static bool GetArtifactSettings(const FString& ArtifactName, const FString& SandboxId, FEOSArtifactSettings& OutSettings);
-	static bool GetArtifactSettingsImpl(const FString& ArtifactName, const TOptional<FString>& SandboxId, FEOSArtifactSettings& OutSettings);
+	static bool GetArtifactSettings(const FString& ArtifactName, const FString& SandboxId, const FString& DeploymentId, FEOSArtifactSettings& OutSettings);
+	static bool GetArtifactSettingsImpl(const FString& ArtifactName, const TOptional<FString>& SandboxId, const TOptional<FString>& DeploymentId, FEOSArtifactSettings& OutSettings);
 
 	static const TArray<FEOSArtifactSettings>& GetCachedArtifactSettings();
 
