Added a static function to FUniqueNetIdEOS to store a constant dedicated server Id value. This is needed for IsValid() to return true on search results.
Modified CopySearchResults to assign the OwningUserId and OwningUserName in session search results for both dedicated and player hosted sessions.This function is now async.
Added a new callback and modified AddSearchResults to account for CopySearchResults being async
Several minor fixes and improvements to session and lobby search.

diff --git a/Engine/Plugins/Online/OnlineSubsystemEOS/Source/OnlineSubsystemEOS/Private/OnlineSessionEOS.cpp b/Engine/Plugins/Online/OnlineSubsystemEOS/Source/OnlineSubsystemEOS/Private/OnlineSessionEOS.cpp
--- a/Engine/Plugins/Online/OnlineSubsystemEOS/Source/OnlineSubsystemEOS/Private/OnlineSessionEOS.cpp
+++ b/Engine/Plugins/Online/OnlineSubsystemEOS/Source/OnlineSubsystemEOS/Private/OnlineSessionEOS.cpp
@@ -664,10 +664,11 @@ void FOnlineSessionEOS::OnSessionInviteReceived(const EOS_Sessions_SessionInvite
 			if (Result == EOS_EResult::EOS_Success)
 			{
 				const TSharedRef<FSessionDetailsEOS> SessionDetails = MakeShared<FSessionDetailsEOS>(SessionDetailsHandle);
-
 				LastInviteSearch = MakeShared<FOnlineSessionSearch>();
-				AddSearchResult(SessionDetails, LastInviteSearch.ToSharedRef());
-				TriggerOnSessionInviteReceivedDelegates(*NetId, *FromNetId, EOSSubsystem->GetAppId(), LastInviteSearch->SearchResults[0]);
+				AddSearchResult(SessionDetails, LastInviteSearch.ToSharedRef(), [this, NetId, FromNetId](bool bWasSuccessful)
+					{
+						TriggerOnSessionInviteReceivedDelegates(*NetId, *FromNetId, EOSSubsystem->GetAppId(), bWasSuccessful ? LastInviteSearch->SearchResults[0] : FOnlineSessionSearchResult());
+					});
 			}
 			else
 			{
@@ -693,10 +694,11 @@ void FOnlineSessionEOS::OnSessionInviteAccepted(const EOS_Sessions_SessionInvite
 			if (Result == EOS_EResult::EOS_Success)
 			{
 				const TSharedRef<FSessionDetailsEOS> SessionDetails = MakeShared<FSessionDetailsEOS>(SessionDetailsHandle);
-
 				LastInviteSearch = MakeShared<FOnlineSessionSearch>();
-				AddSearchResult(SessionDetails, LastInviteSearch.ToSharedRef());
-				TriggerOnSessionUserInviteAcceptedDelegates(true, LocalUserNum, ResolvedUniqueNetId, LastInviteSearch->SearchResults[0]);
+				AddSearchResult(SessionDetails, LastInviteSearch.ToSharedRef(), [this, LocalUserNum, ResolvedUniqueNetId](bool bWasSuccessful)
+					{
+						TriggerOnSessionUserInviteAcceptedDelegates(true, LocalUserNum, ResolvedUniqueNetId, bWasSuccessful ? LastInviteSearch->SearchResults[0] : FOnlineSessionSearchResult());
+					});
 			}
 			else
 			{
@@ -721,10 +723,11 @@ void FOnlineSessionEOS::OnJoinSessionAccepted(const EOS_Sessions_JoinSessionAcce
 			if (Result == EOS_EResult::EOS_Success)
 			{
 				const TSharedRef<FSessionDetailsEOS> SessionDetails = MakeShared<FSessionDetailsEOS>(SessionDetailsHandle);
-
 				LastInviteSearch = MakeShared<FOnlineSessionSearch>();
-				AddSearchResult(SessionDetails, LastInviteSearch.ToSharedRef());
-				TriggerOnSessionUserInviteAcceptedDelegates(true, LocalUserNum, ResolvedUniqueNetId, LastInviteSearch->SearchResults[0]);
+				AddSearchResult(SessionDetails, LastInviteSearch.ToSharedRef(), [this, LocalUserNum, ResolvedUniqueNetId](bool bWasSuccessful)
+					{
+						TriggerOnSessionUserInviteAcceptedDelegates(bWasSuccessful, LocalUserNum, ResolvedUniqueNetId, bWasSuccessful ? LastInviteSearch->SearchResults[0] : FOnlineSessionSearchResult());
+					});
 			}
 			else
 			{
@@ -2307,7 +2310,7 @@ void FOnlineSessionEOS::AddLobbySearchAttribute(EOS_HLobbySearch LobbySearchHand
 	}
 }
 
-void FOnlineSessionEOS::CopySearchResult(const FSessionDetailsEOS& SessionHandle, EOS_SessionDetails_Info* SessionInfo, FOnlineSession& OutSession)
+void FOnlineSessionEOS::CopySearchResult(const FSessionDetailsEOS& SessionHandle, EOS_SessionDetails_Info* SessionInfo, FOnlineSession& OutSession, FOnCopySessionDataCompleteCallback&& Callback)
 {
 	OutSession.NumOpenPrivateConnections = SessionInfo->NumOpenPublicConnections;
 	OutSession.SessionSettings.NumPrivateConnections = SessionInfo->Settings->NumPublicConnections;
@@ -2330,6 +2333,26 @@ void FOnlineSessionEOS::CopySearchResult(const FSessionDetailsEOS& SessionHandle
 	}
 
 	CopyAttributes(SessionHandle, OutSession);
+
+	if (SessionInfo->OwnerServerClientId)
+	{
+		OutSession.OwningUserId = FUniqueNetIdEOS::DedicatedServerId();
+
+		// ResolveUniqueNetIds is an asynchronous operation, so in the cases where it's not called, we'll delay the execution of this callback to match the flow
+		EOSSubsystem->ExecuteNextTick([Callback = MoveTemp(Callback)]()
+			{
+				Callback(true);
+			});
+	}
+	else
+	{
+		EOSSubsystem->UserManager->ResolveUniqueNetId(EOSSubsystem->UserManager->GetDefaultLocalUser(), SessionInfo->OwnerUserId, [this, &OutSession, Callback = MoveTemp(Callback)](FUniqueNetIdEOSRef ResolvedUniqueNetId)
+			{
+				OutSession.OwningUserId = ResolvedUniqueNetId;
+				OutSession.OwningUserName = EOSSubsystem->UserManager->GetPlayerNickname(*ResolvedUniqueNetId);
+				Callback(true);
+			});
+	}
 }
 
 void FOnlineSessionEOS::CopyAttributes(const FSessionDetailsEOS& SessionHandle, FOnlineSession& OutSession)
@@ -2417,7 +2440,7 @@ void FOnlineSessionEOS::CopyAttributes(const FSessionDetailsEOS& SessionHandle,
 	}
 }
 
-void FOnlineSessionEOS::AddSearchResult(const TSharedRef<FSessionDetailsEOS>& SessionHandle, const TSharedRef<FOnlineSessionSearch>& SearchSettings)
+void FOnlineSessionEOS::AddSearchResult(const TSharedRef<FSessionDetailsEOS>& SessionHandle, const TSharedRef<FOnlineSessionSearch>& SearchSettings, FOnCopySessionDataCompleteCallback&& Callback)
 {
 	EOS_SessionDetails_Info* EosSessionDetailsInfo = nullptr;
 	EOS_SessionDetails_CopyInfoOptions CopyOptions = { };
@@ -2433,13 +2456,23 @@ void FOnlineSessionEOS::AddSearchResult(const TSharedRef<FSessionDetailsEOS>& Se
 		TSharedPtr<FOnlineSessionInfoEOS> OnlineSessionInfo = MakeShared<FOnlineSessionInfoEOS>(FOnlineSessionInfoEOS::Create(FUniqueNetIdEOSSession::Create(EosSessionDetailsInfo->SessionId), SessionHandle));
 		SearchResult.Session.SessionInfo = OnlineSessionInfo;
 
-		CopySearchResult(*SessionHandle, EosSessionDetailsInfo, SearchResult.Session);
+		CopySearchResult(*SessionHandle, EosSessionDetailsInfo, SearchResult.Session, MoveTemp(Callback));
 
 		// CopySearchResult above will populate the settings so we can now read the port and construct the HostAddress.
 		OnlineSessionInfo->HostAddr = GetHostAddress(UTF8_TO_TCHAR(EosSessionDetailsInfo->HostAddress), SearchResult.Session.SessionSettings);
 
 		EOS_SessionDetails_Info_Release(EosSessionDetailsInfo);
 	}
+	else
+	{
+		// CopySearchResult may launch an asynchronous operation, so we'll delay the execution of this callback to match the flow
+		EOSSubsystem->ExecuteNextTick([CopyResult, Callback = MoveTemp(Callback)]()
+			{
+				UE_LOG_ONLINE_SESSION(Warning, TEXT("[FOnlineSessionEOS::AddSearchResult]SessionDetails_CopyInfo not successful. Finished with EOS_EResult %s"), *LexToString(CopyResult));
+
+				Callback(false);
+			});
+	}
 }
 
 typedef TEOSCallback<EOS_SessionSearch_OnFindCallback, EOS_SessionSearch_FindCallbackInfo, FOnlineSessionEOS> FFindSessionsCallback;
@@ -2504,31 +2537,67 @@ PRAGMA_ENABLE_DEPRECATION_WARNINGS
 			EOS_SessionSearch_GetSearchResultCountOptions SearchResultOptions = { };
 			SearchResultOptions.ApiVersion = 1;
 			UE_EOS_CHECK_API_MISMATCH(EOS_SESSIONSEARCH_GETSEARCHRESULTCOUNT_API_LATEST, 1);
-			int32 NumSearchResults = EOS_SessionSearch_GetSearchResultCount(CurrentSearchHandle->SearchHandle, &SearchResultOptions);
+			int32 SearchResultsCount = EOS_SessionSearch_GetSearchResultCount(CurrentSearchHandle->SearchHandle, &SearchResultOptions);
 
-			EOS_SessionSearch_CopySearchResultByIndexOptions IndexOptions = { };
-			IndexOptions.ApiVersion = 1;
-			UE_EOS_CHECK_API_MISMATCH(EOS_SESSIONSEARCH_COPYSEARCHRESULTBYINDEX_API_LATEST, 1);
-			for (int32 Index = 0; Index < NumSearchResults; Index++)
+			if (SearchResultsCount > 0)
 			{
-				EOS_HSessionDetails SessionDetailsHandle = nullptr;
-				IndexOptions.SessionIndex = Index;
-				EOS_EResult Result = EOS_SessionSearch_CopySearchResultByIndex(CurrentSearchHandle->SearchHandle, &IndexOptions, &SessionDetailsHandle);
-				if (Result == EOS_EResult::EOS_Success)
+				EOS_SessionSearch_CopySearchResultByIndexOptions IndexOptions = { };
+				IndexOptions.ApiVersion = 1;
+				UE_EOS_CHECK_API_MISMATCH(EOS_SESSIONSEARCH_COPYSEARCHRESULTBYINDEX_API_LATEST, 1);
+				for (int32 SessionIndex = 0; SessionIndex < SearchResultsCount; SessionIndex++)
 				{
-					const TSharedRef<FSessionDetailsEOS> SessionDetails = MakeShared<FSessionDetailsEOS>(SessionDetailsHandle);
+					EOS_HSessionDetails SessionDetailsHandle = nullptr;
+					IndexOptions.SessionIndex = SessionIndex;
+					EOS_EResult Result = EOS_SessionSearch_CopySearchResultByIndex(CurrentSearchHandle->SearchHandle, &IndexOptions, &SessionDetailsHandle);
+					if (Result == EOS_EResult::EOS_Success)
+					{
+						const TSharedRef<FSessionDetailsEOS> SessionDetails = MakeShared<FSessionDetailsEOS>(SessionDetailsHandle);
+						SessionSearchResultsPendingIdResolution.Add(SessionDetails);
+						UE_LOG_ONLINE_SESSION(VeryVerbose, TEXT("[FOnlineSessionEOS::FindEOSSession] SessionSearch_CopySearchResultByIndex was successful."));
+					}
+					else
+					{
+						// It's unlikely EOS_SessionSearch_CopySearchResultByIndex would return failure. If it does return to the caller the failure and stop copying search results
+						UE_LOG_ONLINE_SESSION(Warning, TEXT("[FOnlineSessionEOS::FindEOSSession] SessionSearch_CopySearchResultByIndex not successful. Finished with EOS_EResult %s"), *LexToString(Data->ResultCode));
+						TriggerOnFindSessionsCompleteDelegates(false);
+						return; 
+					}
+				}
 
-					AddSearchResult(SessionDetails, SearchSettings);
+				// Make a copy to iterate over, as the AddSearchResult delegate removes entries.
+				const TArray<TSharedRef<FSessionDetailsEOS>> SessionSearchResultsPendingIdResolutionCopy = SessionSearchResultsPendingIdResolution;
+				
+				// Need to keep track if there is a single failure in AddSearchResult. Assume true and only modify to false if adding a search results returns failure. 
+				bAggregatedAddSearchResultSuccessful = true; 
+				for (const TSharedRef<FSessionDetailsEOS>& SessionDetails : SessionSearchResultsPendingIdResolutionCopy)
+				{
+					AddSearchResult(SessionDetails, SearchSettings, [this, SessionDetails, SearchSettings, SearchResultsCount](bool bWasSuccessful)
+						{
+							SessionSearchResultsPendingIdResolution.Remove(SessionDetails);
+							
+							bAggregatedAddSearchResultSuccessful &= bWasSuccessful; 
+							
+							if (SessionSearchResultsPendingIdResolution.IsEmpty())
+							{
+								SearchSettings->SearchState = bAggregatedAddSearchResultSuccessful ? EOnlineAsyncTaskState::Done : EOnlineAsyncTaskState::Failed; 
+								UE_LOG_ONLINE_SESSION(Log, TEXT("[FOnlineSessionEOS::FindEOSSession] SessionSearch returned %d results. Search state is: %s."), SearchResultsCount, bAggregatedAddSearchResultSuccessful ? TEXT("Done") : TEXT("Failed"));
+								TriggerOnFindSessionsCompleteDelegates(bAggregatedAddSearchResultSuccessful);							
+							}
+						});
 				}
 			}
-			SearchSettings->SearchState = EOnlineAsyncTaskState::Done;
+			else
+			{
+				UE_LOG_ONLINE_SESSION(Log, TEXT("[FOnlineSessionEOS::FindEOSSession] SessionSearch_GetSearchResultCount returned no results"));
+				TriggerOnFindSessionsCompleteDelegates(bWasSuccessful);
+			}
 		}
 		else
 		{
 			SearchSettings->SearchState = EOnlineAsyncTaskState::Failed;
 			UE_LOG_ONLINE_SESSION(Error, TEXT("EOS_SessionSearch_Find() failed with EOS result code (%s)"), *LexToString(Data->ResultCode));
+			TriggerOnFindSessionsCompleteDelegates(bWasSuccessful);
 		}
-		TriggerOnFindSessionsCompleteDelegates(bWasSuccessful);
 	};
 
 	SearchSettings->SearchState = EOnlineAsyncTaskState::InProgress;
@@ -2589,30 +2658,36 @@ void FOnlineSessionEOS::FindEOSSessionById(int32 LocalUserNum, const FUniqueNetI
 			UE_EOS_CHECK_API_MISMATCH(EOS_SESSIONSEARCH_GETSEARCHRESULTCOUNT_API_LATEST, 1);
 			int32 NumSearchResults = EOS_SessionSearch_GetSearchResultCount(CurrentSearchHandle->SearchHandle, &SearchResultOptions);
 
+			// Only a single session is returned when using EOS_SessionSearch_SetSessionId
 			EOS_SessionSearch_CopySearchResultByIndexOptions IndexOptions = { };
 			IndexOptions.ApiVersion = 1;
+			EOS_HSessionDetails SessionDetailsHandle = nullptr;
+			IndexOptions.SessionIndex = 0;
+
 			UE_EOS_CHECK_API_MISMATCH(EOS_SESSIONSEARCH_COPYSEARCHRESULTBYINDEX_API_LATEST, 1);
-			for (int32 Index = 0; Index < NumSearchResults; Index++)
+			EOS_EResult Result = EOS_SessionSearch_CopySearchResultByIndex(CurrentSearchHandle->SearchHandle, &IndexOptions, &SessionDetailsHandle);
+			if (Result == EOS_EResult::EOS_Success)
 			{
-				EOS_HSessionDetails SessionDetailsHandle = nullptr;
-				IndexOptions.SessionIndex = Index;
-				EOS_EResult Result = EOS_SessionSearch_CopySearchResultByIndex(CurrentSearchHandle->SearchHandle, &IndexOptions, &SessionDetailsHandle);
-				if (Result == EOS_EResult::EOS_Success)
-				{
-					const TSharedRef<FSessionDetailsEOS> SessionDetails = MakeShared<FSessionDetailsEOS>(SessionDetailsHandle);
-
-					AddSearchResult(SessionDetails, LocalSessionSearch);
-				}
+				const TSharedRef<FSessionDetailsEOS> SessionDetails = MakeShared<FSessionDetailsEOS>(SessionDetailsHandle);
+				AddSearchResult(SessionDetails, LocalSessionSearch, [this, LocalSessionSearch, LocalUserNum, OnComplete](bool bWasSuccessful)
+					{
+						OnComplete.ExecuteIfBound(LocalUserNum, bWasSuccessful, !LocalSessionSearch->SearchResults.IsEmpty() ? LocalSessionSearch->SearchResults.Last() : FOnlineSessionSearchResult());
+					});
+			}
+			else
+			{
+				// It's unlikely EOS_SessionSearch_CopySearchResultByIndex would return failure. If it does return to the caller the failure.
+				UE_LOG_ONLINE_SESSION(Warning, TEXT("[FOnlineSessionEOS::FindEOSSession] SessionSearch_CopySearchResultByIndex not successful. Finished with EOS_EResult %s"), *LexToString(Data->ResultCode));
+				LocalSessionSearch->SearchState = EOnlineAsyncTaskState::Failed;
+				OnComplete.ExecuteIfBound(LocalUserNum, false, FOnlineSessionSearchResult());
 			}
-			LocalSessionSearch->SearchState = EOnlineAsyncTaskState::Done;
 		}
 		else
 		{
 			LocalSessionSearch->SearchState = EOnlineAsyncTaskState::Failed;
 			UE_LOG_ONLINE_SESSION(Error, TEXT("EOS_SessionSearch_Find() failed with EOS result code (%s)"), *LexToString(Data->ResultCode));
+			OnComplete.ExecuteIfBound(LocalUserNum, false, FOnlineSessionSearchResult());
 		}
-
-		OnComplete.ExecuteIfBound(LocalUserNum, LocalSessionSearch->SearchState == EOnlineAsyncTaskState::Done, !LocalSessionSearch->SearchResults.IsEmpty() ? LocalSessionSearch->SearchResults.Last() : FOnlineSessionSearchResult());
 	};
 
 	EOS_SessionSearch_FindOptions FindOptions = { };
@@ -4616,7 +4691,10 @@ void FOnlineSessionEOS::StartLobbySearch(int32 SearchingPlayerNum, EOS_HLobbySea
 					}
 					else
 					{
+						// It's unlikely EOS_LobbySearch_CopySearchResultByIndex would return failure. If it does return to the caller the failure and stop copying search results
 						UE_LOG_ONLINE_SESSION(Warning, TEXT("[FOnlineSessionEOS::StartLobbySearch::FLobbySearchFindCallback] LobbySearch_CopySearchResultByIndex not successful. Finished with EOS_EResult %s"), *LexToString(Data->ResultCode));
+						CompletionDelegate.ExecuteIfBound(SearchingPlayerNum, false, SearchSettings->SearchResults.Last());
+						return; 
 					}
 				}
 
@@ -4635,11 +4713,6 @@ void FOnlineSessionEOS::StartLobbySearch(int32 SearchingPlayerNum, EOS_HLobbySea
 						}
 					});
 				}
-
-				if (LobbySearchResultsPendingIdResolution.IsEmpty())
-				{
-					CompletionDelegate.ExecuteIfBound(SearchingPlayerNum, true, SearchSettings->SearchResults.Last());
-				}
 			}
 			else
 			{
@@ -4663,7 +4736,7 @@ void FOnlineSessionEOS::StartLobbySearch(int32 SearchingPlayerNum, EOS_HLobbySea
 	EOS_LobbySearch_Find(LobbySearchHandle, &FindOptions, CallbackObj, CallbackObj->GetCallbackPtr());
 }
 
-void FOnlineSessionEOS::AddLobbySearchResult(const TSharedRef<FLobbyDetailsEOS>& LobbyDetails, const TSharedRef<FOnlineSessionSearch>& SearchSettings, const FOnCopyLobbyDataCompleteCallback& Callback)
+void FOnlineSessionEOS::AddLobbySearchResult(const TSharedRef<FLobbyDetailsEOS>& LobbyDetails, const TSharedRef<FOnlineSessionSearch>& SearchSettings, FOnCopyLobbyDataCompleteCallback&& Callback)
 {
 	EOS_LobbyDetails_Info* LobbyDetailsInfo = nullptr;
 	EOS_LobbyDetails_CopyInfoOptions CopyOptions = { };
@@ -4686,7 +4759,7 @@ void FOnlineSessionEOS::AddLobbySearchResult(const TSharedRef<FLobbyDetailsEOS>&
 		SearchResult.Session.SessionInfo = SessionInfo;
 
 		// We copy the lobby data and settings, but not the member data (for search results)
-		CopyLobbyData(LobbyDetails, LobbyDetailsInfo, SearchResult.Session, false, Callback);
+		CopyLobbyData(LobbyDetails, LobbyDetailsInfo, SearchResult.Session, false, MoveTemp(Callback));
 
 		EOS_LobbyDetails_Info_Release(LobbyDetailsInfo);
 
@@ -4695,7 +4768,7 @@ void FOnlineSessionEOS::AddLobbySearchResult(const TSharedRef<FLobbyDetailsEOS>&
 	else
 	{
 		// CopyLobbyData may launch an asynchronous operation, so we'll delay the execution of this callback to match the flow
-		EOSSubsystem->ExecuteNextTick([CopyResult, Callback]()
+		EOSSubsystem->ExecuteNextTick([CopyResult, Callback = MoveTemp(Callback)]()
 		{
 			UE_LOG_ONLINE_SESSION(Warning, TEXT("[FOnlineSessionEOS::AddLobbySearchResult] LobbyDetails_CopyInfo not successful. Finished with EOS_EResult %s"), *LexToString(CopyResult));
 
@@ -4704,7 +4777,7 @@ void FOnlineSessionEOS::AddLobbySearchResult(const TSharedRef<FLobbyDetailsEOS>&
 	}
 }
 
-void FOnlineSessionEOS::CopyLobbyData(const TSharedRef<FLobbyDetailsEOS>& LobbyDetails, EOS_LobbyDetails_Info* LobbyDetailsInfo, FOnlineSession& OutSession, bool bCopyMemberData, const FOnCopyLobbyDataCompleteCallback& Callback)
+void FOnlineSessionEOS::CopyLobbyData(const TSharedRef<FLobbyDetailsEOS>& LobbyDetails, EOS_LobbyDetails_Info* LobbyDetailsInfo, FOnlineSession& OutSession, bool bCopyMemberData, FOnCopyLobbyDataCompleteCallback&& Callback)
 {
 	// This method launches an asynchronous operation, so we'll pass the details handle as a shared ref to make sure it stays alive
 
@@ -4766,7 +4839,7 @@ void FOnlineSessionEOS::CopyLobbyData(const TSharedRef<FLobbyDetailsEOS>& LobbyD
 
 		if (!TargetUserIds.IsEmpty())
 		{
-			EOSSubsystem->UserManager->ResolveUniqueNetIds(EOSSubsystem->UserManager->GetDefaultLocalUser(), TargetUserIds, [this, LobbyDetails, LobbyId = FUniqueNetIdEOSLobby::Create(UTF8_TO_TCHAR(LobbyDetailsInfo->LobbyId)), OriginalCallback = Callback](TMap<EOS_ProductUserId, FUniqueNetIdEOSRef> ResolvedUniqueNetIds)
+			EOSSubsystem->UserManager->ResolveUniqueNetIds(EOSSubsystem->UserManager->GetDefaultLocalUser(), TargetUserIds, [this, LobbyDetails, LobbyId = FUniqueNetIdEOSLobby::Create(UTF8_TO_TCHAR(LobbyDetailsInfo->LobbyId)), Callback = MoveTemp(Callback)](TMap<EOS_ProductUserId, FUniqueNetIdEOSRef> ResolvedUniqueNetIds)
 				{
 					FOnlineSession* Session = GetOnlineSessionFromLobbyId(*LobbyId);
 					if (Session)
@@ -4793,7 +4866,7 @@ void FOnlineSessionEOS::CopyLobbyData(const TSharedRef<FLobbyDetailsEOS>& LobbyD
 					}
 
 					const bool bWasSuccessful = Session != nullptr;
-					OriginalCallback(bWasSuccessful);
+					Callback(bWasSuccessful);
 				});
 
 			return;
@@ -4807,7 +4880,7 @@ void FOnlineSessionEOS::CopyLobbyData(const TSharedRef<FLobbyDetailsEOS>& LobbyD
 
 		const EOS_ProductUserId LobbyOwner = EOS_LobbyDetails_GetLobbyOwner(LobbyDetails->LobbyDetailsHandle, &GetLobbyOwnerOptions);
 
-		EOSSubsystem->UserManager->ResolveUniqueNetIds(EOSSubsystem->UserManager->GetDefaultLocalUser(), { LobbyOwner }, [this, LobbyOwner, LobbyDetails, LobbyId = FUniqueNetIdEOSLobby::Create(UTF8_TO_TCHAR(LobbyDetailsInfo->LobbyId)), OriginalCallback = Callback](TMap<EOS_ProductUserId, FUniqueNetIdEOSRef> ResolvedUniqueNetIds)
+		EOSSubsystem->UserManager->ResolveUniqueNetIds(EOSSubsystem->UserManager->GetDefaultLocalUser(), { LobbyOwner }, [this, LobbyOwner, LobbyDetails, LobbyId = FUniqueNetIdEOSLobby::Create(UTF8_TO_TCHAR(LobbyDetailsInfo->LobbyId)), Callback = MoveTemp(Callback)](TMap<EOS_ProductUserId, FUniqueNetIdEOSRef> ResolvedUniqueNetIds)
 			{
 				FOnlineSession* Session = GetOnlineSessionFromLobbyId(*LobbyId);
 				if (Session)
@@ -4821,14 +4894,14 @@ void FOnlineSessionEOS::CopyLobbyData(const TSharedRef<FLobbyDetailsEOS>& LobbyD
 				}
 
 				const bool bWasSuccessful = Session != nullptr;
-				OriginalCallback(bWasSuccessful);
+				Callback(bWasSuccessful);
 			});
 
 		return;
 	}
 
 	// ResolveUniqueNetIds is an asynchronous operation, so in the cases where it's not called, we'll delay the execution of this callback to match the flow
-	EOSSubsystem->ExecuteNextTick([Callback]()
+	EOSSubsystem->ExecuteNextTick([Callback = MoveTemp(Callback)]()
 	{
 		Callback(true);
 	});
diff --git a/Engine/Plugins/Online/OnlineSubsystemEOS/Source/OnlineSubsystemEOS/Private/OnlineSessionEOS.h b/Engine/Plugins/Online/OnlineSubsystemEOS/Source/OnlineSubsystemEOS/Private/OnlineSessionEOS.h
--- a/Engine/Plugins/Online/OnlineSubsystemEOS/Source/OnlineSubsystemEOS/Private/OnlineSessionEOS.h
+++ b/Engine/Plugins/Online/OnlineSubsystemEOS/Source/OnlineSubsystemEOS/Private/OnlineSessionEOS.h
@@ -252,13 +252,13 @@ private:
 
 	// Methods to update an OSS Lobby from an API Lobby
 	typedef TFunction<void(bool bWasSuccessful)> FOnCopyLobbyDataCompleteCallback;
-	void CopyLobbyData(const TSharedRef<FLobbyDetailsEOS>& LobbyDetails, EOS_LobbyDetails_Info* LobbyDetailsInfo, FOnlineSession& OutSession, bool bCopyMemberData, const FOnCopyLobbyDataCompleteCallback& Callback);
+	void CopyLobbyData(const TSharedRef<FLobbyDetailsEOS>& LobbyDetails, EOS_LobbyDetails_Info* LobbyDetailsInfo, FOnlineSession& OutSession, bool bCopyMemberData, FOnCopyLobbyDataCompleteCallback&& Callback);
 	void CopyLobbyAttributes(const FLobbyDetailsEOS& LobbyDetails, FOnlineSession& OutSession);
 	void CopyLobbyMemberAttributes(const FLobbyDetailsEOS& LobbyDetails, const EOS_ProductUserId& TargetUserId, FSessionSettings& OutSessionSettings);
 
 	// Lobby search
 	void AddLobbySearchAttribute(EOS_HLobbySearch LobbySearchHandle, const EOS_Lobby_AttributeData* Attribute, EOS_EOnlineComparisonOp ComparisonOp);
-	void AddLobbySearchResult(const TSharedRef<FLobbyDetailsEOS>& LobbyDetails, const TSharedRef<FOnlineSessionSearch>& SearchSettings, const FOnCopyLobbyDataCompleteCallback& Callback);
+	void AddLobbySearchResult(const TSharedRef<FLobbyDetailsEOS>& LobbyDetails, const TSharedRef<FOnlineSessionSearch>& SearchSettings, FOnCopyLobbyDataCompleteCallback&& Callback);
 
 	// Helper methods
 	typedef TFunction<void(const EOS_ProductUserId& ProductUserId, EOS_EpicAccountId& EpicAccountId)> GetEpicAccountIdAsyncCallback;
@@ -278,6 +278,8 @@ private:
 	static FString GetBucketId(const FOnlineSessionSearch& SessionSettings);
 
 	// EOS Sessions
+	TArray<TSharedRef<FSessionDetailsEOS>> SessionSearchResultsPendingIdResolution;
+
 	uint32 CreateEOSSession(int32 HostingPlayerNum, FNamedOnlineSession* Session);
 	uint32 JoinEOSSession(int32 PlayerNum, FNamedOnlineSession* Session, const FOnlineSession* SearchSession);
 	uint32 StartEOSSession(FNamedOnlineSession* Session);
@@ -293,9 +295,10 @@ private:
 	void BeginSessionAnalytics(FNamedOnlineSession* Session);
 	void EndSessionAnalytics();
 
-	void AddSearchResult(const TSharedRef<FSessionDetailsEOS>& SessionHandle, const TSharedRef<FOnlineSessionSearch>& SearchSettings);
+	typedef TFunction<void(bool bWasSuccessful)> FOnCopySessionDataCompleteCallback;
+	void AddSearchResult(const TSharedRef<FSessionDetailsEOS>& SessionHandle, const TSharedRef<FOnlineSessionSearch>& SearchSettings, FOnCopySessionDataCompleteCallback&& Callback);
 	void AddSearchAttribute(EOS_HSessionSearch SearchHandle, const EOS_Sessions_AttributeData* Attribute, EOS_EOnlineComparisonOp ComparisonOp);
-	void CopySearchResult(const FSessionDetailsEOS& SessionHandle, EOS_SessionDetails_Info* SessionInfo, FOnlineSession& SessionSettings);
+	void CopySearchResult(const FSessionDetailsEOS& SessionHandle, EOS_SessionDetails_Info* SessionInfo, FOnlineSession& SessionSettings, FOnCopySessionDataCompleteCallback&& Callback);
 	void CopyAttributes(const FSessionDetailsEOS& SessionHandle, FOnlineSession& OutSession);
 
 	void SetPermissionLevel(EOS_HSessionModification SessionModHandle, FNamedOnlineSession* Session);
@@ -334,6 +337,9 @@ private:
 	/** The last accepted invite search. It searches by session id */
 	TSharedPtr<FOnlineSessionSearch> LastInviteSearch;
 
+	/** Used to track adding asynchronous session search results */
+	bool bAggregatedAddSearchResultSuccessful;
+
 	/** Notification state for SDK events */
 	EOS_NotificationId SessionInviteReceivedId;
 	FCallbackBase* SessionInviteReceivedCallback;
diff --git a/Engine/Plugins/Online/OnlineSubsystemEOS/Source/OnlineSubsystemEOS/Private/OnlineSubsystemEOSTypes.cpp b/Engine/Plugins/Online/OnlineSubsystemEOS/Source/OnlineSubsystemEOS/Private/OnlineSubsystemEOSTypes.cpp
--- a/Engine/Plugins/Online/OnlineSubsystemEOS/Source/OnlineSubsystemEOS/Private/OnlineSubsystemEOSTypes.cpp
+++ b/Engine/Plugins/Online/OnlineSubsystemEOS/Source/OnlineSubsystemEOS/Private/OnlineSubsystemEOSTypes.cpp
@@ -7,6 +7,9 @@
 #include "Misc/ScopeRWLock.h"
 #include "UserManagerEOS.h"
 
+#define EOS_DedicatedServer_ID 0xD5D5D5D5D5D5D5D5
+const uint8 EOS_DedicatedServer_ID_RawBytes[EOS_ID_BYTE_SIZE] = "D5D5D5D5D5D5D5D5D5D5D5D5D5D5D5D";
+
 const FUniqueNetIdEOS& FUniqueNetIdEOS::Cast(const FUniqueNetId& NetId)
 {
 	check(GetTypeStatic() == NetId.GetType());
@@ -19,6 +22,12 @@ const FUniqueNetIdEOSRef& FUniqueNetIdEOS::EmptyId()
 	return EmptyId;
 }
 
+const FUniqueNetIdEOSRef& FUniqueNetIdEOS::DedicatedServerId()
+{
+	static const FUniqueNetIdEOSRef DedicatedServerId(Create(EOS_DedicatedServer_ID_RawBytes, EOS_ID_BYTE_SIZE));
+	return DedicatedServerId;
+}
+
 FName FUniqueNetIdEOS::GetTypeStatic()
 {
 	static FName NAME_Eos(TEXT("EOS"));
@@ -40,9 +49,24 @@ int32 FUniqueNetIdEOS::GetSize() const
 	return EOS_ID_BYTE_SIZE;
 }
 
+bool FUniqueNetIdEOS::IsDSValue() const 
+{
+	return ProductUserId == EOS_ProductUserId(EOS_DedicatedServer_ID); 
+}
+
 bool FUniqueNetIdEOS::IsValid() const
 {
-	return EOS_EpicAccountId_IsValid(EpicAccountId) || EOS_ProductUserId_IsValid(ProductUserId);
+	return IsDSValue() || EOS_EpicAccountId_IsValid(EpicAccountId) || EOS_ProductUserId_IsValid(ProductUserId);
+}
+
+const EOS_EpicAccountId FUniqueNetIdEOS::GetEpicAccountId() const
+{
+	return EpicAccountId;
+}
+
+const  EOS_ProductUserId FUniqueNetIdEOS::GetProductUserId() const
+{
+	return IsDSValue() ? nullptr: ProductUserId;
 }
 
 uint32 FUniqueNetIdEOS::GetTypeHash() const
@@ -74,18 +98,26 @@ FUniqueNetIdEOS::FUniqueNetIdEOS(const uint8* Bytes, int32 Size)
 	check(Size == EOS_ID_BYTE_SIZE);
 	FMemory::Memcpy(RawBytes, Bytes, EOS_ID_BYTE_SIZE);
 
-	const bool bIsEasNonZero = Algo::AnyOf(TArrayView<const uint8>(Bytes, ID_HALF_BYTE_SIZE));
-	if (bIsEasNonZero)
+	if (FMemory::Memcmp(RawBytes, EOS_DedicatedServer_ID_RawBytes, EOS_ID_BYTE_SIZE) == 0)
 	{
-		const FString EpicAccountIdStr = BytesToHexLower(Bytes, ID_HALF_BYTE_SIZE);
-		EpicAccountId = EOS_EpicAccountId_FromString(TCHAR_TO_UTF8(*EpicAccountIdStr));
+		EpicAccountId = nullptr;
+		ProductUserId = EOS_ProductUserId(EOS_DedicatedServer_ID);
 	}
-
-	const bool bIsPuidNonZero = Algo::AnyOf(TArrayView<const uint8>(Bytes + ID_HALF_BYTE_SIZE, ID_HALF_BYTE_SIZE));
-	if (bIsPuidNonZero)
+	else
 	{
-		const FString ProductUserIdStr = BytesToHexLower(Bytes + ID_HALF_BYTE_SIZE, ID_HALF_BYTE_SIZE);
-		ProductUserId = EOS_ProductUserId_FromString(TCHAR_TO_UTF8(*ProductUserIdStr));
+		const bool bIsEasNonZero = Algo::AnyOf(TArrayView<const uint8>(Bytes, ID_HALF_BYTE_SIZE));
+		if (bIsEasNonZero)
+		{
+			const FString EpicAccountIdStr = BytesToHexLower(Bytes, ID_HALF_BYTE_SIZE);
+			EpicAccountId = EOS_EpicAccountId_FromString(TCHAR_TO_UTF8(*EpicAccountIdStr));
+		}
+
+		const bool bIsPuidNonZero = Algo::AnyOf(TArrayView<const uint8>(Bytes + ID_HALF_BYTE_SIZE, ID_HALF_BYTE_SIZE));
+		if (bIsPuidNonZero)
+		{
+			const FString ProductUserIdStr = BytesToHexLower(Bytes + ID_HALF_BYTE_SIZE, ID_HALF_BYTE_SIZE);
+			ProductUserId = EOS_ProductUserId_FromString(TCHAR_TO_UTF8(*ProductUserIdStr));
+		}
 	}
 }
 
@@ -133,6 +165,10 @@ FUniqueNetIdEOSRef FUniqueNetIdEOSRegistry::FindOrAddImpl(const uint8* Bytes, in
 	{
 		// The net id type already knows how to deserialize, so create a temp one to get the EAS/PUID
 		FUniqueNetIdEOS Temp(Bytes, Size);
+		if (Temp.IsDSValue())
+		{
+			return FUniqueNetIdEOS::DedicatedServerId(); 
+		}
 		return FindOrAddImpl(Temp.GetEpicAccountId(), Temp.GetProductUserId());
 	}
 	return FUniqueNetIdEOS::EmptyId();
diff --git a/Engine/Plugins/Online/OnlineSubsystemEOS/Source/OnlineSubsystemEOS/Private/OnlineSubsystemEOSTypes.h b/Engine/Plugins/Online/OnlineSubsystemEOS/Source/OnlineSubsystemEOS/Private/OnlineSubsystemEOSTypes.h
--- a/Engine/Plugins/Online/OnlineSubsystemEOS/Source/OnlineSubsystemEOS/Private/OnlineSubsystemEOSTypes.h
+++ b/Engine/Plugins/Online/OnlineSubsystemEOS/Source/OnlineSubsystemEOS/Private/OnlineSubsystemEOSTypes.h
@@ -42,26 +42,20 @@ public:
 
 	/** global static instance of invalid (zero) id */
 	static const FUniqueNetIdEOSRef& EmptyId();
+	static const FUniqueNetIdEOSRef& DedicatedServerId();
 
 	static FName GetTypeStatic();
 
 	virtual FName GetType() const override;
 	virtual const uint8* GetBytes() const override;
 	virtual int32 GetSize() const override;
+
 	virtual bool IsValid() const override;
 	virtual uint32 GetTypeHash() const override;
 	virtual FString ToString() const override;
 	virtual FString ToDebugString() const override;
-
-	virtual const EOS_EpicAccountId GetEpicAccountId() const override
-	{
-		return EpicAccountId;
-	}
-
-	virtual const EOS_ProductUserId GetProductUserId() const override
-	{
-		return ProductUserId;
-	}
+	virtual const EOS_EpicAccountId GetEpicAccountId() const override;
+	virtual const EOS_ProductUserId GetProductUserId() const override;
 
 private:
 	EOS_EpicAccountId EpicAccountId = nullptr;
@@ -78,6 +72,7 @@ private:
 
 	FUniqueNetIdEOS() = default;
 	
+	bool IsDSValue() const;
 	explicit FUniqueNetIdEOS(const uint8* Bytes, int32 Size);
 	explicit FUniqueNetIdEOS(EOS_EpicAccountId InEpicAccountId, EOS_ProductUserId InProductUserId);
 };
