PR #11779: Don't attempt a tick at shutdown if EOSShared module is no longer loaded.

diff --git a/Engine/Plugins/Online/OnlineSubsystemEOS/Source/OnlineSubsystemEOS/Private/OnlineSubsystemEOS.cpp b/Engine/Plugins/Online/OnlineSubsystemEOS/Source/OnlineSubsystemEOS/Private/OnlineSubsystemEOS.cpp
--- a/Engine/Plugins/Online/OnlineSubsystemEOS/Source/OnlineSubsystemEOS/Private/OnlineSubsystemEOS.cpp
+++ b/Engine/Plugins/Online/OnlineSubsystemEOS/Source/OnlineSubsystemEOS/Private/OnlineSubsystemEOS.cpp
@@ -442,7 +442,12 @@ bool FOnlineSubsystemEOS::Shutdown()
 	// EOS-22677 workaround: Make sure tick is called at least once before shutting down.
 	if (EOSPlatformHandle)
 	{
-		EOS_Platform_Tick(*EOSPlatformHandle);
+		// The EOSShared module may have been shut down at this point, in which case
+		// the handle has been released and we can't use it.
+		if (FModuleManager::Get().IsModuleLoaded("EOSShared"))
+		{
+			EOS_Platform_Tick(*EOSPlatformHandle);
+		}
 	}
 
 	StopTicker();
