- Add -noeos support to OnlineSubsystemEOS

diff --git a/Engine/Plugins/Online/OnlineSubsystemEOS/Source/OnlineSubsystemEOS/Private/OnlineSubsystemEOSModule.cpp b/Engine/Plugins/Online/OnlineSubsystemEOS/Source/OnlineSubsystemEOS/Private/OnlineSubsystemEOSModule.cpp
--- a/Engine/Plugins/Online/OnlineSubsystemEOS/Source/OnlineSubsystemEOS/Private/OnlineSubsystemEOSModule.cpp
+++ b/Engine/Plugins/Online/OnlineSubsystemEOS/Source/OnlineSubsystemEOS/Private/OnlineSubsystemEOSModule.cpp
@@ -10,8 +10,9 @@
 #include "Features/IModularFeature.h"
 #include "Features/IModularFeatures.h"
 
-#include "Misc/CoreDelegates.h"
+#include "Misc/CommandLine.h"
 #include "Misc/ConfigCacheIni.h"
+#include "Misc/CoreDelegates.h"
 #include "Misc/LazySingleton.h"
 #include "Modules/ModuleInterface.h"
 
@@ -90,6 +91,11 @@ IMPLEMENT_MODULE(FOnlineSubsystemEOSModule, OnlineSubsystemEOS);
 
 void FOnlineSubsystemEOSModule::StartupModule()
 {
+	if (FParse::Param(FCommandLine::Get(), TEXT("NoEOS")))
+	{
+		return;
+	}
+
 	EOSFactory = new FOnlineFactoryEOS();
 
 	// Create and register our singleton factory with the main online subsystem for easy access
