- Fix delete of nullptr when -noeos is passed
- Minor cleanup, move from raw pointer to TUniquePtr.

diff --git a/Engine/Plugins/Online/OnlineSubsystemEOS/Source/OnlineSubsystemEOS/Private/OnlineSubsystemEOSModule.cpp b/Engine/Plugins/Online/OnlineSubsystemEOS/Source/OnlineSubsystemEOS/Private/OnlineSubsystemEOSModule.cpp
--- a/Engine/Plugins/Online/OnlineSubsystemEOS/Source/OnlineSubsystemEOS/Private/OnlineSubsystemEOSModule.cpp
+++ b/Engine/Plugins/Online/OnlineSubsystemEOS/Source/OnlineSubsystemEOS/Private/OnlineSubsystemEOSModule.cpp
@@ -54,18 +54,10 @@ public:
  */
 class FOnlineSubsystemEOSModule : public IModuleInterface
 {
-private:
-
-	/** Class responsible for creating instance(s) of the subsystem */
-	FOnlineFactoryEOS* EOSFactory;
-
 public:
+	FOnlineSubsystemEOSModule() = default;
 
-	FOnlineSubsystemEOSModule() :
-		EOSFactory(NULL)
-	{}
-
-	virtual ~FOnlineSubsystemEOSModule() {}
+	virtual ~FOnlineSubsystemEOSModule() = default;
 
 #if WITH_EDITOR
 	void OnPostEngineInit();
@@ -85,6 +77,10 @@ public:
 	{
 		return false;
 	}
+
+private:
+	/** Class responsible for creating instance(s) of the subsystem */
+	TUniquePtr<FOnlineFactoryEOS> EOSFactory;
 };
 
 IMPLEMENT_MODULE(FOnlineSubsystemEOSModule, OnlineSubsystemEOS);
@@ -96,11 +92,11 @@ void FOnlineSubsystemEOSModule::StartupModule()
 		return;
 	}
 
-	EOSFactory = new FOnlineFactoryEOS();
+	EOSFactory = MakeUnique<FOnlineFactoryEOS>();
 
 	// Create and register our singleton factory with the main online subsystem for easy access
 	FOnlineSubsystemModule& OSS = FModuleManager::GetModuleChecked<FOnlineSubsystemModule>("OnlineSubsystem");
-	OSS.RegisterPlatformService(EOS_SUBSYSTEM, EOSFactory);
+	OSS.RegisterPlatformService(EOS_SUBSYSTEM, EOSFactory.Get());
 
 #if WITH_EOS_SDK
 	// Have to call this as early as possible in order to hook the rendering device
@@ -150,8 +146,7 @@ void FOnlineSubsystemEOSModule::ShutdownModule()
 	FOnlineSubsystemModule& OSS = FModuleManager::GetModuleChecked<FOnlineSubsystemModule>("OnlineSubsystem");
 	OSS.UnregisterPlatformService(EOS_SUBSYSTEM);
 
-	delete EOSFactory;
-	EOSFactory = nullptr;
+	EOSFactory.Reset();
 
 	TLazySingleton<FUniqueNetIdEOSRegistry>::TearDown();
 }
