Refactor of obsolete method call in UserManagerEOS

diff --git a/Engine/Plugins/Online/OnlineSubsystemEOS/Source/OnlineSubsystemEOS/Private/UserManagerEOS.cpp b/Engine/Plugins/Online/OnlineSubsystemEOS/Source/OnlineSubsystemEOS/Private/UserManagerEOS.cpp
--- a/Engine/Plugins/Online/OnlineSubsystemEOS/Source/OnlineSubsystemEOS/Private/UserManagerEOS.cpp
+++ b/Engine/Plugins/Online/OnlineSubsystemEOS/Source/OnlineSubsystemEOS/Private/UserManagerEOS.cpp
@@ -703,7 +703,9 @@ void FUserManagerEOS::LoginViaPersistentAuth(int32 LocalUserNum)
 	}
 
 	// If we don't have an auth token saved or it's expired, we'll get one via AccountPortal login first
-	LoginViaAccountPortal(LocalUserNum);
+	FOnlineAccountCredentials Credentials;
+	Credentials.Type = TEXT("accountportal");
+	CallEOSAuthLogin(LocalUserNum, Credentials, false);
 #endif
 }
 
