Fix crash in OnlineSubsystemEOS when accepting a friend request

diff --git a/Engine/Plugins/Online/OnlineSubsystemEOS/Source/OnlineSubsystemEOS/Private/UserManagerEOS.cpp b/Engine/Plugins/Online/OnlineSubsystemEOS/Source/OnlineSubsystemEOS/Private/UserManagerEOS.cpp
--- a/Engine/Plugins/Online/OnlineSubsystemEOS/Source/OnlineSubsystemEOS/Private/UserManagerEOS.cpp
+++ b/Engine/Plugins/Online/OnlineSubsystemEOS/Source/OnlineSubsystemEOS/Private/UserManagerEOS.cpp
@@ -2796,31 +2796,37 @@ void FUserManagerEOS::FriendStatusChangedImpl(EOS_EpicAccountId LocalUserId, EOS
 
 FOnlineFriendEOSRef FUserManagerEOS::AddFriend(int32 LocalUserNum, const FUniqueNetIdEOS& FriendNetId)
 {
-	// A call to AddFriend should only be made after the friends list has been initialised
-	FFriendsListEOSPtr FriendsListPtr = GetLocalUserChecked(LocalUserNum).FriendsList;
-	check(FriendsListPtr);
-
 	const FOnlineUserEOSRef UserRef = UniqueNetIdToUserRefMap[FriendNetId.AsShared()];
 	const FUniqueNetIdEOSRef FriendNetIdEOSRef = StaticCastSharedRef<const FUniqueNetIdEOS>(FriendNetId.AsShared());
 	const FOnlineFriendEOSRef FriendRef = MakeShareable(new FOnlineFriendEOS(FriendNetIdEOSRef, UserRef->UserAttributes, *EOSSubsystem));
 
-	FriendsListPtr->Add(FriendNetId.AsShared(), FriendRef);
+	// A call to AddFriend should only be made after the friends list has been initialised
+	FFriendsListEOSPtr FriendsListPtr = GetLocalUserChecked(LocalUserNum).FriendsList;
+	if (FriendsListPtr.IsValid())
+	{
+		FriendsListPtr->Add(FriendNetId.AsShared(), FriendRef);
 
-	EOS_Friends_GetStatusOptions Options = { };
-	Options.ApiVersion = 1;
-	UE_EOS_CHECK_API_MISMATCH(EOS_FRIENDS_GETSTATUS_API_LATEST, 1);
-	Options.LocalUserId = GetLocalEpicAccountId(LocalUserNum);
-	Options.TargetUserId = FriendNetId.GetEpicAccountId();
-	EOS_EFriendsStatus Status = EOS_Friends_GetStatus(EOSSubsystem->FriendsHandle, &Options);
-	
-	FriendRef->SetInviteStatus(ToEInviteStatus(Status));
+		EOS_Friends_GetStatusOptions Options = { };
+		Options.ApiVersion = 1;
+		UE_EOS_CHECK_API_MISMATCH(EOS_FRIENDS_GETSTATUS_API_LATEST, 1);
+		Options.LocalUserId = GetLocalEpicAccountId(LocalUserNum);
+		Options.TargetUserId = FriendNetId.GetEpicAccountId();
+		EOS_EFriendsStatus Status = EOS_Friends_GetStatus(EOSSubsystem->FriendsHandle, &Options);
 
-	// Querying the presence of a non-friend would cause an SDK error.
-	// Players that sent/received a friend invitation from us still count as "friends", so check
-	// our friend relationship here.
-	if(Status == EOS_EFriendsStatus::EOS_FS_Friends)
+		FriendRef->SetInviteStatus(ToEInviteStatus(Status));
+
+		// Querying the presence of a non-friend would cause an SDK error.
+		// Players that sent/received a friend invitation from us still count as "friends", so check
+		// our friend relationship here.
+		if (Status == EOS_EFriendsStatus::EOS_FS_Friends)
+		{
+			QueryPresence(FriendNetId, IgnoredPresenceDelegate);
+		}
+	}
+	else
 	{
-		QueryPresence(FriendNetId, IgnoredPresenceDelegate);
+		// Friends list ptr was not valid
+		UE_LOG_ONLINE_FRIEND(Warning, TEXT("AddFriend() failed. FriendsListPtr is not valid.  Was ReadFriendsList() run first?"));
 	}
 
 	return FriendRef;
