Refactor of PersistentAuth/AccountPortal login flow in UserManagerEOS.
Addition of missing logic for proper EOS login flow on Switch when using NPLN.

diff --git a/Engine/Plugins/Online/OnlineSubsystemEOS/Source/OnlineSubsystemEOS/Private/UserManagerEOS.cpp b/Engine/Plugins/Online/OnlineSubsystemEOS/Source/OnlineSubsystemEOS/Private/UserManagerEOS.cpp
--- a/Engine/Plugins/Online/OnlineSubsystemEOS/Source/OnlineSubsystemEOS/Private/UserManagerEOS.cpp
+++ b/Engine/Plugins/Online/OnlineSubsystemEOS/Source/OnlineSubsystemEOS/Private/UserManagerEOS.cpp
@@ -550,7 +550,7 @@ bool FUserManagerEOS::Login(int32 LocalUserNum, const FOnlineAccountCredentials&
 		return true;
 	}
 
-	CallEOSAuthLogin(LocalUserNum, AccountCredentials, false);	
+	CallEOSAuthLogin(LocalUserNum, AccountCredentials, false);
 	return true;
 }
 
@@ -619,9 +619,9 @@ void FUserManagerEOS::LoginViaPersistentAuthLegacy(int32 LocalUserNum, const FOn
 				EOSHelpers->GetSystemAuthCredentialsOptions(Credentials.SystemAuthCredentialsOptions);
 
 				FLoginCallback* PACallbackObj = new FLoginCallback(AsWeak());
-				PACallbackObj->CallbackLambda = [this, LocalUserNum, AccountCredentials](const EOS_Auth_LoginCallbackInfo* Data)
+				PACallbackObj->CallbackLambda = [this, LocalUserNum, AccountCredentials, LoginCredentialType = Credentials.Type](const EOS_Auth_LoginCallbackInfo* Data)
 				{
-					OnEOSAuthLoginCompleteLegacy(LocalUserNum, AccountCredentials, true, true, Data);
+					OnEOSAuthLoginCompleteLegacy(LocalUserNum, AccountCredentials, LoginCredentialType, true, Data);
 				};
 
 				EOS_Auth_Login(EOSSubsystem->AuthHandle, &PALoginOptions, (void*)PACallbackObj, PACallbackObj->GetCallbackPtr());
@@ -705,7 +705,7 @@ void FUserManagerEOS::LoginViaPersistentAuth(int32 LocalUserNum)
 	// If we don't have an auth token saved or it's expired, we'll get one via AccountPortal login first
 	FOnlineAccountCredentials Credentials;
 	Credentials.Type = TEXT("accountportal");
-	CallEOSAuthLogin(LocalUserNum, Credentials, false);
+	CallEOSAuthLogin(LocalUserNum, Credentials, true);
 #endif
 }
 
@@ -753,7 +753,7 @@ void FUserManagerEOS::LoginViaExternalAuth(int32 LocalUserNum)
 						return;
 					}
 
-					CallEOSAuthLogin(LocalUserNum, AccountCredentials, false);
+					CallEOSAuthLogin(LocalUserNum, AccountCredentials, true);
 				}
 			}));
 }
@@ -788,15 +788,15 @@ void FUserManagerEOS::LoginViaAccountPortalLegacy(int32 LocalUserNum, const FOnl
 	EOSHelpers->GetSystemAuthCredentialsOptions(EOSCredentials.SystemAuthCredentialsOptions);
 
 	FLoginCallback* CallbackObj = new FLoginCallback(AsWeak());
-	CallbackObj->CallbackLambda = [this, LocalUserNum, AccountCredentials](const EOS_Auth_LoginCallbackInfo* Data) mutable
+	CallbackObj->CallbackLambda = [this, LocalUserNum, AccountCredentials, LoginCredentialType = EOSCredentials.Type](const EOS_Auth_LoginCallbackInfo* Data) mutable
 	{
-		OnEOSAuthLoginCompleteLegacy(LocalUserNum, AccountCredentials, true, false, Data);
+		OnEOSAuthLoginCompleteLegacy(LocalUserNum, AccountCredentials, LoginCredentialType, false, Data);
 	};
 
 	EOS_Auth_Login(EOSSubsystem->AuthHandle, &LoginOptions, (void*)CallbackObj, CallbackObj->GetCallbackPtr());
 }
 
-void FUserManagerEOS::CallEOSAuthLogin(int32 LocalUserNum, const FOnlineAccountCredentials& Credentials, bool bRetryLogin)
+void FUserManagerEOS::CallEOSAuthLogin(int32 LocalUserNum, const FOnlineAccountCredentials& Credentials, bool bIsAutoLogin)
 {
 	// First we construct the EOS Credentials object
 	EOS_Auth_Credentials EOSCredentials = {};
@@ -859,20 +859,16 @@ void FUserManagerEOS::CallEOSAuthLogin(int32 LocalUserNum, const FOnlineAccountC
 	FPlatformEOSHelpersPtr EOSHelpers = EOSSubsystem->GetEOSHelpers();
 	EOSHelpers->GetSystemAuthCredentialsOptions(EOSCredentials.SystemAuthCredentialsOptions);
 
-	// Store selection of persistent auth.
-	// The persistent auth token is handled by the EOSSDK. On a login failure the persistent token may need to be deleted if it is invalid.
-	const bool bIsPersistentLogin = EOSCredentials.Type == EOS_ELoginCredentialType::EOS_LCT_PersistentAuth;
-
 	FLoginCallback* CallbackObj = new FLoginCallback(AsWeak());
-	CallbackObj->CallbackLambda = [this, LocalUserNum, Credentials, bIsPersistentLogin, bRetryLogin, bUseNewLoginFlow = Settings.bUseNewLoginFlow](const EOS_Auth_LoginCallbackInfo* Data)
+	CallbackObj->CallbackLambda = [this, LocalUserNum, Credentials, LoginCredentialType = EOSCredentials.Type, bIsAutoLogin, bUseNewLoginFlow = Settings.bUseNewLoginFlow](const EOS_Auth_LoginCallbackInfo* Data)
 	{
 		if (bUseNewLoginFlow)
 		{
-			OnEOSAuthLoginComplete(LocalUserNum, bIsPersistentLogin, bRetryLogin, Data);
+			OnEOSAuthLoginComplete(LocalUserNum, LoginCredentialType, bIsAutoLogin, Data);
 		}
 		else
 		{
-			OnEOSAuthLoginCompleteLegacy(LocalUserNum, Credentials, bIsPersistentLogin, bRetryLogin, Data);
+			OnEOSAuthLoginCompleteLegacy(LocalUserNum, Credentials, LoginCredentialType, bIsAutoLogin, Data);
 		}
 	};
 
@@ -925,11 +921,15 @@ void FUserManagerEOS::CopyAndSaveEpicAuthToken(int32 LocalUserNum, const EOS_Epi
 	}
 }
 
-void FUserManagerEOS::OnEOSAuthLoginCompleteLegacy(int32 LocalUserNum, const FOnlineAccountCredentials& Credentials, bool bIsPersistentLogin, bool bRetryLogin, const EOS_Auth_LoginCallbackInfo* Data)
+void FUserManagerEOS::OnEOSAuthLoginCompleteLegacy(int32 LocalUserNum, const FOnlineAccountCredentials& Credentials, const EOS_ELoginCredentialType LoginCredentialType, const bool bIsAutoLogin, const EOS_Auth_LoginCallbackInfo* Data)
 {
 	if (Data->ResultCode == EOS_EResult::EOS_Success)
 	{
-		if (bIsPersistentLogin)
+		const bool bShouldSaveToken =
+			LoginCredentialType == EOS_ELoginCredentialType::EOS_LCT_PersistentAuth ||
+			LoginCredentialType == EOS_ELoginCredentialType::EOS_LCT_AccountPortal;
+			
+		if (bShouldSaveToken)
 		{
 #if EOS_AUTH_TOKEN_SAVEGAME_STORAGE
 			CopyAndSaveEpicAuthToken(LocalUserNum, Data->LocalUserId);
@@ -945,7 +945,9 @@ void FUserManagerEOS::OnEOSAuthLoginCompleteLegacy(int32 LocalUserNum, const FOn
 		// Link the account
 		LinkEASLegacy(LocalUserNum, Data->ContinuanceToken, Credentials);
 	}
-	else if (Data->ResultCode == EOS_EResult::EOS_InvalidAuth && bIsPersistentLogin && bRetryLogin)
+	else if (Data->ResultCode == EOS_EResult::EOS_InvalidAuth
+		&& LoginCredentialType == EOS_ELoginCredentialType::EOS_LCT_PersistentAuth
+		&& bIsAutoLogin)
 	{
 		// We attempted a Persistent Auth login but there is no stored token, we'll attempt Account Portal login automatically
 		LoginViaAccountPortalLegacy(LocalUserNum, Credentials);
@@ -959,12 +961,15 @@ void FUserManagerEOS::OnEOSAuthLoginCompleteLegacy(int32 LocalUserNum, const FOn
 			TriggerOnLoginCompleteDelegates(LocalUserNum, false, *FUniqueNetIdEOS::EmptyId(), ErrorString);
 		};
 
-		const bool bShouldRemoveCachedToken =
+		const bool bIsCachedTokenFailure =
 			Data->ResultCode == EOS_EResult::EOS_AccessDenied ||
 			Data->ResultCode == EOS_EResult::EOS_Auth_InvalidToken;
-
-		// Check for invalid persistent login credentials.
-		if (bIsPersistentLogin && bShouldRemoveCachedToken)
+			
+		const bool bShouldRemoveCachedToken =
+			LoginCredentialType == EOS_ELoginCredentialType::EOS_LCT_PersistentAuth &&
+			bIsCachedTokenFailure;
+			
+		if (bShouldRemoveCachedToken)
 		{
 			FDeletePersistentAuthCallback* DeleteAuthCallbackObj = new FDeletePersistentAuthCallback(AsWeak());
 			DeleteAuthCallbackObj->CallbackLambda = [this, LocalUserNum, TriggerLoginFailure](const EOS_Auth_DeletePersistentAuthCallbackInfo* Data)
@@ -998,11 +1003,15 @@ void FUserManagerEOS::OnEOSAuthLoginCompleteLegacy(int32 LocalUserNum, const FOn
 	}
 }
 
-void FUserManagerEOS::OnEOSAuthLoginComplete(int32 LocalUserNum, bool bIsPersistentLogin, bool bRetryLogin, const EOS_Auth_LoginCallbackInfo* Data)
+void FUserManagerEOS::OnEOSAuthLoginComplete(int32 LocalUserNum, const EOS_ELoginCredentialType LoginCredentialType, const bool bIsAutoLogin, const EOS_Auth_LoginCallbackInfo* Data)
 {
 	if (Data->ResultCode == EOS_EResult::EOS_Success)
 	{
-		if (bIsPersistentLogin)
+		const bool bShouldSaveToken =
+			LoginCredentialType == EOS_ELoginCredentialType::EOS_LCT_PersistentAuth ||
+			LoginCredentialType == EOS_ELoginCredentialType::EOS_LCT_AccountPortal;
+			
+		if (bShouldSaveToken)
 		{
 #if EOS_AUTH_TOKEN_SAVEGAME_STORAGE
 			CopyAndSaveEpicAuthToken(LocalUserNum, Data->LocalUserId);
@@ -1018,35 +1027,40 @@ void FUserManagerEOS::OnEOSAuthLoginComplete(int32 LocalUserNum, bool bIsPersist
 		// Link the account
 		LinkEAS(LocalUserNum, Data->ContinuanceToken);
 	}
-	else if (Data->ResultCode == EOS_EResult::EOS_InvalidAuth && bIsPersistentLogin && bRetryLogin)
+	else if (Data->ResultCode == EOS_EResult::EOS_InvalidAuth
+		&& LoginCredentialType == EOS_ELoginCredentialType::EOS_LCT_PersistentAuth
+		&& bIsAutoLogin)
 	{
 		// We attempted a Persistent Auth login but there is no stored token, we'll attempt Account Portal login automatically (but won't retry if that fails)
 		FOnlineAccountCredentials Credentials;
 		Credentials.Type = TEXT("accountportal");
-		CallEOSAuthLogin(LocalUserNum, Credentials, false);
+		CallEOSAuthLogin(LocalUserNum, Credentials, true);
 	}
 	else
 	{
 		auto TriggerLoginFailure = [this, LocalUserNum, LoginResultCode = Data->ResultCode]()
-			{
-				FString ErrorString = FString::Printf(TEXT("Login(%d) failed with EOS result code (%s)"), LocalUserNum, *LexToString(LoginResultCode));
-				UE_LOG_ONLINE(Warning, TEXT("%s"), *ErrorString);
-				TriggerOnLoginCompleteDelegates(LocalUserNum, false, *FUniqueNetIdEOS::EmptyId(), ErrorString);
-			};
+		{
+			FString ErrorString = FString::Printf(TEXT("Login(%d) failed with EOS result code (%s)"), LocalUserNum, *LexToString(LoginResultCode));
+			UE_LOG_ONLINE(Warning, TEXT("%s"), *ErrorString);
+			TriggerOnLoginCompleteDelegates(LocalUserNum, false, *FUniqueNetIdEOS::EmptyId(), ErrorString);
+		};
 
-		const bool bShouldRemoveCachedToken =
+		const bool bIsCachedTokenFailure =
 			Data->ResultCode == EOS_EResult::EOS_AccessDenied ||
 			Data->ResultCode == EOS_EResult::EOS_Auth_InvalidToken;
-
-		// Check for invalid persistent login credentials.
-		if (bIsPersistentLogin && bShouldRemoveCachedToken)
+			
+		const bool bShouldRemoveCachedToken =
+			LoginCredentialType == EOS_ELoginCredentialType::EOS_LCT_PersistentAuth &&
+			bIsCachedTokenFailure;
+			
+		if (bShouldRemoveCachedToken)
 		{
 			FDeletePersistentAuthCallback* DeleteAuthCallbackObj = new FDeletePersistentAuthCallback(AsWeak());
 			DeleteAuthCallbackObj->CallbackLambda = [this, LocalUserNum, TriggerLoginFailure](const EOS_Auth_DeletePersistentAuthCallbackInfo* Data)
-				{
-					// Deleting the auth token is best effort.
-					TriggerLoginFailure();
-				};
+			{
+				// Deleting the auth token is best effort.
+				TriggerLoginFailure();
+			};
 
 			EOS_Auth_DeletePersistentAuthOptions DeletePersistentAuthOptions;
 			DeletePersistentAuthOptions.ApiVersion = 2;
diff --git a/Engine/Plugins/Online/OnlineSubsystemEOS/Source/OnlineSubsystemEOS/Private/UserManagerEOS.h b/Engine/Plugins/Online/OnlineSubsystemEOS/Source/OnlineSubsystemEOS/Private/UserManagerEOS.h
--- a/Engine/Plugins/Online/OnlineSubsystemEOS/Source/OnlineSubsystemEOS/Private/UserManagerEOS.h
+++ b/Engine/Plugins/Online/OnlineSubsystemEOS/Source/OnlineSubsystemEOS/Private/UserManagerEOS.h
@@ -406,14 +406,14 @@ private:
 	void LinkEASLegacy(int32 LocalUserNum, EOS_ContinuanceToken Token, const FOnlineAccountCredentials& AccountCredentials);
 	bool ConnectLoginNoEASLegacy(int32 LocalUserNum, const FOnlineAccountCredentials& AccountCredentials);
 	void FullLoginCallbackLegacy(int32 LocalUserNum, EOS_EpicAccountId AccountId, EOS_ProductUserId UserId, const FOnlineAccountCredentials& AccountCredentials);
-	void OnEOSAuthLoginCompleteLegacy(int32 LocalUserNum, const FOnlineAccountCredentials& Credentials, bool bIsPersistentLogin, bool bRetryLogin, const EOS_Auth_LoginCallbackInfo* Data);
+	void OnEOSAuthLoginCompleteLegacy(int32 LocalUserNum, const FOnlineAccountCredentials& Credentials, const EOS_ELoginCredentialType LoginCredentialType, const bool bIsAutoLogin, const EOS_Auth_LoginCallbackInfo* Data);
 
 	bool IsLocalUserValid(int32 LocalUserNum) const ;
 	FLocalUserEOS& GetLocalUserChecked(int32 LocalUserNum);
 
-	void CallEOSAuthLogin(int32 LocalUserNum, const FOnlineAccountCredentials& Credentials, bool bRetryLogin);
+	void CallEOSAuthLogin(int32 LocalUserNum, const FOnlineAccountCredentials& Credentials, bool bIsAutoLogin);
 	void CopyAndSaveEpicAuthToken(int32 LocalUserNum, const EOS_EpicAccountId& EpicAccountId);
-	void OnEOSAuthLoginComplete(int32 LocalUserNum, bool bIsPersistentLogin, bool bRetryLogin, const EOS_Auth_LoginCallbackInfo* Data);
+	void OnEOSAuthLoginComplete(int32 LocalUserNum, const EOS_ELoginCredentialType LoginCredentialType, const bool bIsAutoLogin, const EOS_Auth_LoginCallbackInfo* Data);
 
 	void RemoveLocalUser(int32 LocalUserNum);
 	FLocalUserEOS& AddLocalUser(int32 LocalUserNum, EOS_EpicAccountId EpicAccountId, EOS_ProductUserId UserId);
