Fixes for incorrectly used FString constructors.

diff --git a/Engine/Plugins/Online/OnlineSubsystemEOS/Source/OnlineSubsystemEOS/Private/UserManagerEOS.cpp b/Engine/Plugins/Online/OnlineSubsystemEOS/Source/OnlineSubsystemEOS/Private/UserManagerEOS.cpp
--- a/Engine/Plugins/Online/OnlineSubsystemEOS/Source/OnlineSubsystemEOS/Private/UserManagerEOS.cpp
+++ b/Engine/Plugins/Online/OnlineSubsystemEOS/Source/OnlineSubsystemEOS/Private/UserManagerEOS.cpp
@@ -3095,7 +3095,7 @@ bool FUserManagerEOS::SendInvite(int32 LocalUserNum, const FUniqueNetId& FriendI
 	if (!LocalUsers.IsValidIndex(LocalUserNum))
 	{
 		UE_LOG_ONLINE_FRIEND(Warning, TEXT("Can't SendInvite() for user (%d) since they are not logged in"), LocalUserNum);
-		Delegate.ExecuteIfBound(LocalUserNum, false, FriendId, ListName, FString(TEXT("Can't SendInvite() for user (%d) since they are not logged in"), LocalUserNum));
+		Delegate.ExecuteIfBound(LocalUserNum, false, FriendId, ListName, FString::Printf(TEXT("Can't SendInvite() for user (%d) since they are not logged in"), LocalUserNum));
 		return false;
 	}
 
@@ -3104,7 +3104,7 @@ bool FUserManagerEOS::SendInvite(int32 LocalUserNum, const FUniqueNetId& FriendI
 	if (EOS_EpicAccountId_IsValid(AccountId) == EOS_FALSE)
 	{
 		UE_LOG_ONLINE_FRIEND(Warning, TEXT("Can't SendInvite() for user (%d) since the potential player id is unknown"), LocalUserNum);
-		Delegate.ExecuteIfBound(LocalUserNum, false, FriendId, ListName, FString(TEXT("Can't SendInvite() for user (%d) since the player id is unknown"), LocalUserNum));
+		Delegate.ExecuteIfBound(LocalUserNum, false, FriendId, ListName, FString::Printf(TEXT("Can't SendInvite() for user (%d) since the player id is unknown"), LocalUserNum));
 		return false;
 	}
 
@@ -3140,7 +3140,7 @@ bool FUserManagerEOS::AcceptInvite(int32 LocalUserNum, const FUniqueNetId& Frien
 	if (!LocalUsers.IsValidIndex(LocalUserNum))
 	{
 		UE_LOG_ONLINE_FRIEND(Warning, TEXT("Can't AcceptInvite() for user (%d) since they are not logged in"), LocalUserNum);
-		Delegate.ExecuteIfBound(LocalUserNum, false, FriendId, ListName, FString(TEXT("Can't AcceptInvite() for user (%d) since they are not logged in"), LocalUserNum));
+		Delegate.ExecuteIfBound(LocalUserNum, false, FriendId, ListName, FString::Printf(TEXT("Can't AcceptInvite() for user (%d) since they are not logged in"), LocalUserNum));
 		return false;
 	}
 
@@ -3149,7 +3149,7 @@ bool FUserManagerEOS::AcceptInvite(int32 LocalUserNum, const FUniqueNetId& Frien
 	if (EOS_EpicAccountId_IsValid(AccountId) == EOS_FALSE)
 	{
 		UE_LOG_ONLINE_FRIEND(Warning, TEXT("Can't AcceptInvite() for user (%d) since the friend is not in their list"), LocalUserNum);
-		Delegate.ExecuteIfBound(LocalUserNum, false, FriendId, ListName, FString(TEXT("Can't AcceptInvite() for user (%d) since the friend is not in their list"), LocalUserNum));
+		Delegate.ExecuteIfBound(LocalUserNum, false, FriendId, ListName, FString::Printf(TEXT("Can't AcceptInvite() for user (%d) since the friend is not in their list"), LocalUserNum));
 		return false;
 	}
 
