PR #12022:

Added missing parameter to EOS_LobbyDetails_CopyMemberAttributeByIndex call in OnlineSessionEOS

diff --git a/Engine/Plugins/Online/OnlineSubsystemEOS/Source/OnlineSubsystemEOS/Private/OnlineSessionEOS.cpp b/Engine/Plugins/Online/OnlineSubsystemEOS/Source/OnlineSubsystemEOS/Private/OnlineSessionEOS.cpp
--- a/Engine/Plugins/Online/OnlineSubsystemEOS/Source/OnlineSubsystemEOS/Private/OnlineSessionEOS.cpp
+++ b/Engine/Plugins/Online/OnlineSubsystemEOS/Source/OnlineSubsystemEOS/Private/OnlineSessionEOS.cpp
@@ -5004,6 +5004,7 @@ void FOnlineSessionEOS::CopyLobbyMemberAttributes(const FLobbyDetailsEOS& LobbyD
 		AttrOptions.ApiVersion = 1;
 		UE_EOS_CHECK_API_MISMATCH(EOS_LOBBYDETAILS_COPYMEMBERATTRIBUTEBYINDEX_API_LATEST, 1);
 		AttrOptions.AttrIndex = MemberAttributeIndex;
+		AttrOptions.TargetUserId = TargetUserId;
 
 		EOS_Lobby_Attribute* Attribute = NULL;
 		EOS_EResult ResultCode = EOS_LobbyDetails_CopyMemberAttributeByIndex(LobbyDetails.LobbyDetailsHandle, &AttrOptions, &Attribute);
