- Fix broken log macro

#rnx

diff --git a/Engine/Plugins/Online/OnlineSubsystemEOS/Source/OnlineSubsystemEOS/Private/OnlineSubsystemEOS.cpp b/Engine/Plugins/Online/OnlineSubsystemEOS/Source/OnlineSubsystemEOS/Private/OnlineSubsystemEOS.cpp
--- a/Engine/Plugins/Online/OnlineSubsystemEOS/Source/OnlineSubsystemEOS/Private/OnlineSubsystemEOS.cpp
+++ b/Engine/Plugins/Online/OnlineSubsystemEOS/Source/OnlineSubsystemEOS/Private/OnlineSubsystemEOS.cpp
@@ -262,14 +262,14 @@ bool FOnlineSubsystemEOS::Init()
 	bool bUnused;
 	if (GConfig->GetBool(TEXT("/Script/OnlineSubsystemEOS.EOSSettings"), TEXT("bShouldEnforceBeingLaunchedByEGS"), bUnused, GEngineIni))
 	{
-		UE_LOG_ONLINE(Error, TEXT("%hs: Support for bShouldEnforceBeingLaunchedByEGS has been removed, please delete this config entry and instead set bUseLauncherChecks=true in your .Target.cs file(s)"));
+		UE_LOG_ONLINE(Error, TEXT("FOnlineSubsystemEOS: Support for bShouldEnforceBeingLaunchedByEGS has been removed, please delete this config entry and instead set bUseLauncherChecks=true in your .Target.cs file(s)"));
 		return false;
 	}
 
 	EOSSDKManager = IEOSSDKManager::Get();
 	if (!EOSSDKManager)
 	{
-		UE_LOG_ONLINE(Error, TEXT("FOnlineSubsystemEOS::Init() failed to get EOSSDKManager interface"));
+		UE_LOG_ONLINE(Error, TEXT("FOnlineSubsystemEOS: failed to get EOSSDKManager interface"));
 		return false;
 	}
 
