Proper fix for the friends acceptance crash in OnlineSubsystemEOS.  Drop the friend change event completely instead of going into the AddFriend method.

diff --git a/Engine/Plugins/Online/OnlineSubsystemEOS/Source/OnlineSubsystemEOS/Private/UserManagerEOS.cpp b/Engine/Plugins/Online/OnlineSubsystemEOS/Source/OnlineSubsystemEOS/Private/UserManagerEOS.cpp
--- a/Engine/Plugins/Online/OnlineSubsystemEOS/Source/OnlineSubsystemEOS/Private/UserManagerEOS.cpp
+++ b/Engine/Plugins/Online/OnlineSubsystemEOS/Source/OnlineSubsystemEOS/Private/UserManagerEOS.cpp
@@ -2729,6 +2729,14 @@ void FUserManagerEOS::FriendStatusChanged(const EOS_Friends_OnFriendsUpdateInfo*
 void FUserManagerEOS::FriendStatusChangedImpl(EOS_EpicAccountId LocalUserId, EOS_EpicAccountId TargetUserId, EOS_EFriendsStatus PreviousStatus, EOS_EFriendsStatus CurrentStatus)
 {
 	const int32 LocalUserNum = GetLocalUserNumFromEpicAccountId(LocalUserId);
+	// Check to see if the friends list is valid before proceeding.  If not, drop the event.
+	FFriendsListEOSPtr FriendsListPtr = GetLocalUserChecked(LocalUserNum).FriendsList;
+	if (!FriendsListPtr.IsValid())
+	{
+		UE_LOG_ONLINE_FRIEND(Verbose, TEXT("FriendStatusChangedImpl: Friend list still has not been queried for local user [%d]. Friends status change event will be dropped."), LocalUserNum);
+		return;
+	}
+
 	const FUniqueNetIdEOSRef LocalEOSID = GetLocalUserChecked(LocalUserNum).UniqueNetId.ToSharedRef();
 	const FUniqueNetIdEOSRef& FriendEOSId = FUniqueNetIdEOSRegistry::FindChecked(TargetUserId);
 
@@ -2737,32 +2745,24 @@ void FUserManagerEOS::FriendStatusChangedImpl(EOS_EpicAccountId LocalUserId, EOS
 	case EOS_EFriendsStatus::EOS_FS_NotFriends: // Invite rejections and friend removal
 	{
 		//User should already be a friend
-		FFriendsListEOSPtr& FriendsListPtr = GetLocalUserChecked(LocalUserNum).FriendsList;
-		if (FriendsListPtr)
+		FOnlineFriendEOSPtr Friend = FriendsListPtr->GetByNetId(FriendEOSId);
+		if (Friend.IsValid())
 		{
-			FOnlineFriendEOSPtr Friend = FriendsListPtr->GetByNetId(FriendEOSId);
-			if (Friend.IsValid())
-			{
-				FriendsListPtr->Remove(FriendEOSId, Friend.ToSharedRef());
-				Friend->SetInviteStatus(EInviteStatus::Unknown);
+			FriendsListPtr->Remove(FriendEOSId, Friend.ToSharedRef());
+			Friend->SetInviteStatus(EInviteStatus::Unknown);
 
-				if (PreviousStatus == EOS_EFriendsStatus::EOS_FS_Friends)
-				{
-					TriggerOnFriendRemovedDelegates(*LocalEOSID, *FriendEOSId);
-				}
-				else if (PreviousStatus == EOS_EFriendsStatus::EOS_FS_InviteSent || PreviousStatus == EOS_EFriendsStatus::EOS_FS_InviteReceived)
-				{
-					TriggerOnInviteRejectedDelegates(*LocalEOSID, *FriendEOSId); // We don't have an "OnInviteRejected" event only for the local user
-				}
+			if (PreviousStatus == EOS_EFriendsStatus::EOS_FS_Friends)
+			{
+				TriggerOnFriendRemovedDelegates(*LocalEOSID, *FriendEOSId);
 			}
-			else
+			else if (PreviousStatus == EOS_EFriendsStatus::EOS_FS_InviteSent || PreviousStatus == EOS_EFriendsStatus::EOS_FS_InviteReceived)
 			{
-				UE_LOG_ONLINE_FRIEND(Verbose, TEXT("Friend status notification received for user [%d], but remote user [%s] was not previously registered as a friend"), LocalUserNum, *FriendEOSId->ToString());
+				TriggerOnInviteRejectedDelegates(*LocalEOSID, *FriendEOSId); // We don't have an "OnInviteRejected" event only for the local user
 			}
 		}
 		else
 		{
-			UE_LOG_ONLINE_FRIEND(Verbose, TEXT("Friend list still has not been queried for local user [%d]"), LocalUserNum);
+			UE_LOG_ONLINE_FRIEND(Verbose, TEXT("Friend status notification received for user [%d], but remote user [%s] was not previously registered as a friend"), LocalUserNum, *FriendEOSId->ToString());
 		}
 
 		break;
@@ -2816,31 +2816,24 @@ FOnlineFriendEOSRef FUserManagerEOS::AddFriend(int32 LocalUserNum, const FUnique
 
 	// A call to AddFriend should only be made after the friends list has been initialised
 	FFriendsListEOSPtr FriendsListPtr = GetLocalUserChecked(LocalUserNum).FriendsList;
-	if (FriendsListPtr.IsValid())
-	{
-		FriendsListPtr->Add(FriendNetId.AsShared(), FriendRef);
+	check(FriendsListPtr);
+	FriendsListPtr->Add(FriendNetId.AsShared(), FriendRef);
 
-		EOS_Friends_GetStatusOptions Options = { };
-		Options.ApiVersion = 1;
-		UE_EOS_CHECK_API_MISMATCH(EOS_FRIENDS_GETSTATUS_API_LATEST, 1);
-		Options.LocalUserId = GetLocalEpicAccountId(LocalUserNum);
-		Options.TargetUserId = FriendNetId.GetEpicAccountId();
-		EOS_EFriendsStatus Status = EOS_Friends_GetStatus(EOSSubsystem->FriendsHandle, &Options);
+	EOS_Friends_GetStatusOptions Options = { };
+	Options.ApiVersion = 1;
+	UE_EOS_CHECK_API_MISMATCH(EOS_FRIENDS_GETSTATUS_API_LATEST, 1);
+	Options.LocalUserId = GetLocalEpicAccountId(LocalUserNum);
+	Options.TargetUserId = FriendNetId.GetEpicAccountId();
+	EOS_EFriendsStatus Status = EOS_Friends_GetStatus(EOSSubsystem->FriendsHandle, &Options);
 
-		FriendRef->SetInviteStatus(ToEInviteStatus(Status));
+	FriendRef->SetInviteStatus(ToEInviteStatus(Status));
 
-		// Querying the presence of a non-friend would cause an SDK error.
-		// Players that sent/received a friend invitation from us still count as "friends", so check
-		// our friend relationship here.
-		if (Status == EOS_EFriendsStatus::EOS_FS_Friends)
-		{
-			QueryPresence(FriendNetId, IgnoredPresenceDelegate);
-		}
-	}
-	else
+	// Querying the presence of a non-friend would cause an SDK error.
+	// Players that sent/received a friend invitation from us still count as "friends", so check
+	// our friend relationship here.
+	if (Status == EOS_EFriendsStatus::EOS_FS_Friends)
 	{
-		// Friends list ptr was not valid
-		UE_LOG_ONLINE_FRIEND(Warning, TEXT("AddFriend() failed. FriendsListPtr is not valid.  Was ReadFriendsList() run first?"));
+		QueryPresence(FriendNetId, IgnoredPresenceDelegate);
 	}
 
 	return FriendRef;
