- Update EOSSettings to override the found settings with the sandboxid / deploymentid passed on command line.

[RN] Updated EOSSettings to override the found settings with the sandboxid / deploymentid passed on command line.

diff --git a/Engine/Plugins/Online/OnlineSubsystemEOS/Source/OnlineSubsystemEOS/Private/EOSSettings.cpp b/Engine/Plugins/Online/OnlineSubsystemEOS/Source/OnlineSubsystemEOS/Private/EOSSettings.cpp
--- a/Engine/Plugins/Online/OnlineSubsystemEOS/Source/OnlineSubsystemEOS/Private/EOSSettings.cpp
+++ b/Engine/Plugins/Online/OnlineSubsystemEOS/Source/OnlineSubsystemEOS/Private/EOSSettings.cpp
@@ -262,31 +262,45 @@ bool UEOSSettings::GetSelectedArtifactSettings(FEOSArtifactSettings& OutSettings
 	// Prefer -EpicDeploymentIdOverride over previous.
 	bHasDeploymentId |= FParse::Value(FCommandLine::Get(), TEXT("EpicDeploymentIdOverride="), DeploymentId);
 
-	// If present, grab the settings where both match
+	bool bSettingsFound = false;
+
+	// If SandboxId and DeploymentId are both specified, look for settings with matching ArtifactName, SandboxId, and DeploymentId
 	if (bHasSandboxId && bHasDeploymentId)
 	{
-		if (GetArtifactSettings(ArtifactName, SandboxId, DeploymentId, OutSettings))
-		{
-			return true;
-		}
+		bSettingsFound = GetArtifactSettings(ArtifactName, SandboxId, DeploymentId, OutSettings);
+		UE_CLOG_ONLINE(!bSettingsFound, Verbose, TEXT("%hs ArtifactName=[%s] SandboxId=[%s] DeploymentId=[%s] no settings found for trio, falling back on pair check."),
+			__FUNCTION__, *ArtifactName, *SandboxId, *DeploymentId);
+	}
 
-		UE_LOG_ONLINE(Log, TEXT("UEOSSettings::GetSelectedArtifactSettings() ArtifactName=[%s] SandboxId=[%s] DeploymentId=[%s] no settings found for trio, falling back on pair check."), *ArtifactName, *SandboxId, *DeploymentId);
+	// Fall back on settings with matching ArtifactName and SandboxId
+	if (!bSettingsFound && bHasSandboxId)
+	{
+		bSettingsFound = GetArtifactSettings(ArtifactName, SandboxId, OutSettings);
+		UE_CLOG_ONLINE(!bSettingsFound, Verbose, TEXT("%hs ArtifactName=[%s] SandboxId=[%s] no settings found for pair, falling back on just ArtifactName."),
+			__FUNCTION__, *ArtifactName, *SandboxId);
 	}
 
-	if (bHasSandboxId)
+	// Fall back on settings with matching ArtifactName. This assumes non-EGS and only one settings entry per ArtifactName in config.
+	if (!bSettingsFound)
 	{
-		if (GetArtifactSettings(ArtifactName, SandboxId, OutSettings))
+		bSettingsFound = GetArtifactSettings(ArtifactName, OutSettings);
+		UE_CLOG_ONLINE(!bSettingsFound, Error, TEXT("UEOSSettings::GetSelectedArtifactSettings() ArtifactName=[%s] no settings found."), *ArtifactName);
+	}
+
+	// Override the config settings with the SandboxId and/or DeploymentId passed on command line.
+	if (bSettingsFound)
+	{
+		if (bHasSandboxId)
 		{
-			return true;
+			OutSettings.SandboxId = SandboxId;
+		}
+		if (bHasDeploymentId)
+		{
+			OutSettings.DeploymentId = DeploymentId;
 		}
-
-		UE_LOG_ONLINE(Log, TEXT("UEOSSettings::GetSelectedArtifactSettings() ArtifactName=[%s] SandboxId=[%s] no settings found for pair, falling back on just ArtifactName."), *ArtifactName, *SandboxId);
 	}
 
-	// Fall back on just matching the Artifact name. This assumes non-EGS and only one settings entry per ArtifactName in config.
-	const bool bSuccess = GetArtifactSettings(ArtifactName, OutSettings);
-	UE_CLOG_ONLINE(!bSuccess, Error, TEXT("UEOSSettings::GetSelectedArtifactSettings() ArtifactName=[%s] no settings found."), *ArtifactName);
-	return bSuccess;
+	return bSettingsFound;
 }
 
 FString UEOSSettings::GetDefaultArtifactName()
