- Support adding a ArtifactName-less "wildcard" config entry

This sits at the bottom of the preference hierarchy and gets picked up (if it exists) if nothing else matches. This enables licensees to specify a single Artifacts entry, with no ArtifactName, that gets used in all cases.

diff --git a/Engine/Plugins/Online/OnlineSubsystemEOS/Source/OnlineSubsystemEOS/Private/EOSSettings.cpp b/Engine/Plugins/Online/OnlineSubsystemEOS/Source/OnlineSubsystemEOS/Private/EOSSettings.cpp
--- a/Engine/Plugins/Online/OnlineSubsystemEOS/Source/OnlineSubsystemEOS/Private/EOSSettings.cpp
+++ b/Engine/Plugins/Online/OnlineSubsystemEOS/Source/OnlineSubsystemEOS/Private/EOSSettings.cpp
@@ -264,6 +264,19 @@ bool UEOSSettings::GetSelectedArtifactSettings(FEOSArtifactSettings& OutSettings
 
 	bool bSettingsFound = false;
 
+	// Find the config. We have a hierarchy for what to use, depending on what arguments are provided
+	//
+	// 1. If SandboxId and DeploymentId are provided on command line, find config entry with matching ArtifactName, SandboxId and DeploymentId
+	// 2. If we didn't find a config entry, and SandboxId is provided on command line, find config entry with matching ArtifactName and SandboxId
+	// 3. If we didn't find a config entry, find config entry with matching ArtifactName
+	// 4. If we didn't find a config entry, find config entry with empty ArtifactName
+	//
+	// Note for most use cases it is sufficient to ignore 1/2/3 and just provide a single artifact config entry with empty ArtifactName,
+	// in which case the client id etc specified in that entry will be used in all cases.
+	// SandboxId/DeploymentId provided on command line will take precedence over those specified in the config entry.
+	// To support running outside of EGS, ensure you provide values for SandboxId and DeploymentId in the artifact config,
+	// and DefaultArtifactName in EOSSettings config, to use when -EpicApp, -EpicSandboxId and/or -EpicDeploymentId are not provided.
+
 	// If SandboxId and DeploymentId are both specified, look for settings with matching ArtifactName, SandboxId, and DeploymentId
 	if (bHasSandboxId && bHasDeploymentId)
 	{
@@ -280,16 +293,28 @@ bool UEOSSettings::GetSelectedArtifactSettings(FEOSArtifactSettings& OutSettings
 			__FUNCTION__, *ArtifactName, *SandboxId);
 	}
 
-	// Fall back on settings with matching ArtifactName. This assumes non-EGS and only one settings entry per ArtifactName in config.
+	// Fall back on settings with matching ArtifactName.
 	if (!bSettingsFound)
 	{
 		bSettingsFound = GetArtifactSettings(ArtifactName, OutSettings);
-		UE_CLOG_ONLINE(!bSettingsFound, Error, TEXT("UEOSSettings::GetSelectedArtifactSettings() ArtifactName=[%s] no settings found."), *ArtifactName);
+		UE_CLOG_ONLINE(!bSettingsFound, Verbose, TEXT("%hs ArtifactName=[%s] no settings found for ArtifactName, falling back on empty ArtifactName."),
+			__FUNCTION__, *ArtifactName);
+	}
+
+	// Fall back on settings with an empty ArtifactName.
+	if (!bSettingsFound)
+	{
+		bSettingsFound = GetArtifactSettings(FString(), OutSettings);
+		UE_CLOG_ONLINE(!bSettingsFound, Verbose, TEXT("%hs No settings found for empty ArtifactName"), __FUNCTION__);
 	}
 
-	// Override the config settings with the SandboxId and/or DeploymentId passed on command line.
+	UE_CLOG_ONLINE(!bSettingsFound, Error, TEXT("%hs ArtifactName=[%s] SandboxId=[%s] DeploymentId=[%s] no settings found."),
+		__FUNCTION__, *ArtifactName, *SandboxId, *DeploymentId);
+
+	// Override the found config with command line values
 	if (bSettingsFound)
 	{
+		OutSettings.ArtifactName = ArtifactName;
 		if (bHasSandboxId)
 		{
 			OutSettings.SandboxId = SandboxId;
