Removed -EpicPortal requirement to initialize Ecom interface in OnlineSubsystemEOS.
An EOS app could be launched from a non-Epic store and still need access to Ecom functionality.

diff --git a/Engine/Plugins/Online/OnlineSubsystemEOS/Source/OnlineSubsystemEOS/Private/OnlineSubsystemEOS.cpp b/Engine/Plugins/Online/OnlineSubsystemEOS/Source/OnlineSubsystemEOS/Private/OnlineSubsystemEOS.cpp
--- a/Engine/Plugins/Online/OnlineSubsystemEOS/Source/OnlineSubsystemEOS/Private/OnlineSubsystemEOS.cpp
+++ b/Engine/Plugins/Online/OnlineSubsystemEOS/Source/OnlineSubsystemEOS/Private/OnlineSubsystemEOS.cpp
@@ -257,8 +257,6 @@ bool FOnlineSubsystemEOS::PlatformCreate()
 
 bool FOnlineSubsystemEOS::Init()
 {
-	bWasLaunchedByEGS = FParse::Param(FCommandLine::Get(), TEXT("EpicPortal"));
-
 	bool bUnused;
 	if (GConfig->GetBool(TEXT("/Script/OnlineSubsystemEOS.EOSSettings"), TEXT("bShouldEnforceBeingLaunchedByEGS"), bUnused, GEngineIni))
 	{
@@ -351,16 +349,11 @@ bool FOnlineSubsystemEOS::Init()
 		UE_LOG_ONLINE(Error, TEXT("FOnlineSubsystemEOS: failed to init EOS platform, couldn't get achievements handle"));
 		return false;
 	}
-	// Disable ecom if not part of EGS
-	if (bWasLaunchedByEGS)
+	EcomHandle = EOS_Platform_GetEcomInterface(*EOSPlatformHandle);
+	if (EcomHandle == nullptr)
 	{
-		EcomHandle = EOS_Platform_GetEcomInterface(*EOSPlatformHandle);
-		if (EcomHandle == nullptr)
-		{
-			UE_LOG_ONLINE(Error, TEXT("FOnlineSubsystemEOS: failed to init EOS platform, couldn't get ecom handle"));
-			return false;
-		}
-		StoreInterfacePtr = MakeShareable(new FOnlineStoreEOS(this));
+		UE_LOG_ONLINE(Error, TEXT("FOnlineSubsystemEOS: failed to init EOS platform, couldn't get ecom handle"));
+		return false;
 	}
 
 	// We set the product id
@@ -421,6 +414,7 @@ bool FOnlineSubsystemEOS::Init()
 	StatsInterfacePtr = MakeShareable(new FOnlineStatsEOS(this));
 	LeaderboardsInterfacePtr = MakeShareable(new FOnlineLeaderboardsEOS(this));
 	AchievementsInterfacePtr = MakeShareable(new FOnlineAchievementsEOS(this));
+	StoreInterfacePtr = MakeShareable(new FOnlineStoreEOS(this));
 
 	// We initialized ok so we can tick
 	StartTicker();
@@ -594,7 +588,6 @@ FOnlineSubsystemEOS::FOnlineSubsystemEOS(FName InInstanceName) :
 	, StoreInterfacePtr(nullptr)
 	, TitleFileInterfacePtr(nullptr)
 	, UserCloudInterfacePtr(nullptr)
-	, bWasLaunchedByEGS(false)
 {
 	StopTicker();
 }
diff --git a/Engine/Plugins/Online/OnlineSubsystemEOS/Source/OnlineSubsystemEOS/Private/OnlineSubsystemEOS.h b/Engine/Plugins/Online/OnlineSubsystemEOS/Source/OnlineSubsystemEOS/Private/OnlineSubsystemEOS.h
--- a/Engine/Plugins/Online/OnlineSubsystemEOS/Source/OnlineSubsystemEOS/Private/OnlineSubsystemEOS.h
+++ b/Engine/Plugins/Online/OnlineSubsystemEOS/Source/OnlineSubsystemEOS/Private/OnlineSubsystemEOS.h
@@ -148,8 +148,6 @@ public:
 	/** User Cloud interface pointer */
 	FOnlineUserCloudEOSPtr UserCloudInterfacePtr;
 
-	bool bWasLaunchedByEGS;
-
 	TSharedPtr<FSocketSubsystemEOS, ESPMode::ThreadSafe> SocketSubsystem;
 
 	static FPlatformEOSHelpersPtr EOSHelpersPtr;
