Wrong completion delegate fired in FOnlineSessionEOS::FindFriendSession
#jire UE-222169

diff --git a/Engine/Plugins/Online/OnlineSubsystemEOS/Source/OnlineSubsystemEOS/Private/OnlineSessionEOS.cpp b/Engine/Plugins/Online/OnlineSubsystemEOS/Source/OnlineSubsystemEOS/Private/OnlineSessionEOS.cpp
--- a/Engine/Plugins/Online/OnlineSubsystemEOS/Source/OnlineSubsystemEOS/Private/OnlineSessionEOS.cpp
+++ b/Engine/Plugins/Online/OnlineSubsystemEOS/Source/OnlineSubsystemEOS/Private/OnlineSessionEOS.cpp
@@ -2978,7 +2978,7 @@ bool FOnlineSessionEOS::FindFriendSession(int32 LocalUserNum, const FUniqueNetId
 
 		StartLobbySearch(LocalUserNum, LobbySearchHandle, CurrentSessionSearch.ToSharedRef(), FOnSingleSessionResultCompleteDelegate::CreateLambda([this](int32 LocalUserNum, bool bWasSuccessful, const FOnlineSessionSearchResult& EOSResult)
 		{
-			TriggerOnFindSessionsCompleteDelegates(bWasSuccessful);
+			TriggerOnFindFriendSessionCompleteDelegates(LocalUserNum, bWasSuccessful, {EOSResult});
 		}));
 
 		bResult = true;
@@ -2986,9 +2986,9 @@ bool FOnlineSessionEOS::FindFriendSession(int32 LocalUserNum, const FUniqueNetId
 	else
 	{
 		UE_LOG_ONLINE_SESSION(Warning, TEXT("[FOnlineSessionEOS::FindFriendSession] CreateLobbySearch not successful. Finished with EOS_EResult %s"), *LexToString(CreateLobbySearchResult));
-		EOSSubsystem->ExecuteNextTick([this]()
+		EOSSubsystem->ExecuteNextTick([this, LocalUserNum]()
 		{
-			TriggerOnFindSessionsCompleteDelegates(false);
+			TriggerOnFindFriendSessionCompleteDelegates(LocalUserNum, false, {});
 		});
 	}
 
@@ -3005,8 +3005,7 @@ bool FOnlineSessionEOS::FindFriendSession(const FUniqueNetId& LocalUserId, const
 	EOSSubsystem->ExecuteNextTick([this, LocalUserIdRef = LocalUserId.AsShared()]()
 		{
 			// this function has to exist due to interface definition, but it does not have a meaningful implementation in EOS subsystem yet
-			TArray<FOnlineSessionSearchResult> EmptySearchResult;
-			TriggerOnFindFriendSessionCompleteDelegates(EOSSubsystem->UserManager->GetLocalUserNumFromUniqueNetId(*LocalUserIdRef), false, EmptySearchResult);
+			TriggerOnFindFriendSessionCompleteDelegates(EOSSubsystem->UserManager->GetLocalUserNumFromUniqueNetId(*LocalUserIdRef), false, {});
 		});
 
 	return true;
