- Fix bShouldEnforceBeingLaunchedByEGS causing errors when set to false, only error when it's set to true (directing people to use LauncherChecks instead)

[RN] Fix startup error in OnlineSubsystemEOS when bShouldEnforceBeingLaunchedByEGS=false is present in config

diff --git a/Engine/Plugins/Online/OnlineSubsystemEOS/Source/OnlineSubsystemEOS/Private/OnlineSubsystemEOS.cpp b/Engine/Plugins/Online/OnlineSubsystemEOS/Source/OnlineSubsystemEOS/Private/OnlineSubsystemEOS.cpp
--- a/Engine/Plugins/Online/OnlineSubsystemEOS/Source/OnlineSubsystemEOS/Private/OnlineSubsystemEOS.cpp
+++ b/Engine/Plugins/Online/OnlineSubsystemEOS/Source/OnlineSubsystemEOS/Private/OnlineSubsystemEOS.cpp
@@ -257,8 +257,9 @@ bool FOnlineSubsystemEOS::PlatformCreate()
 
 bool FOnlineSubsystemEOS::Init()
 {
-	bool bUnused;
-	if (GConfig->GetBool(TEXT("/Script/OnlineSubsystemEOS.EOSSettings"), TEXT("bShouldEnforceBeingLaunchedByEGS"), bUnused, GEngineIni))
+	bool bShouldEnforceBeingLaunchedByEGS = false;
+	GConfig->GetBool(TEXT("/Script/OnlineSubsystemEOS.EOSSettings"), TEXT("bShouldEnforceBeingLaunchedByEGS"), bShouldEnforceBeingLaunchedByEGS, GEngineIni);
+	if (bShouldEnforceBeingLaunchedByEGS)
 	{
 		UE_LOG_ONLINE(Error, TEXT("FOnlineSubsystemEOS: Support for bShouldEnforceBeingLaunchedByEGS has been removed, please delete this config entry and instead set bUseLauncherChecks=true in your .Target.cs file(s)"));
 		return false;
