Fixes in FUserManagerEOS::GetLoginStatus and OnlineAchievementsEOS

diff --git a/Engine/Plugins/Online/OnlineSubsystemEOS/Source/OnlineSubsystemEOS/Private/OnlineAchievementsEOS.cpp b/Engine/Plugins/Online/OnlineSubsystemEOS/Source/OnlineSubsystemEOS/Private/OnlineAchievementsEOS.cpp
--- a/Engine/Plugins/Online/OnlineSubsystemEOS/Source/OnlineSubsystemEOS/Private/OnlineAchievementsEOS.cpp
+++ b/Engine/Plugins/Online/OnlineSubsystemEOS/Source/OnlineSubsystemEOS/Private/OnlineAchievementsEOS.cpp
@@ -59,13 +59,13 @@ void FOnlineAchievementsEOS::UnlockAchievements(const FUniqueNetId& PlayerId, FO
 		bool bWasSuccessful = Data->ResultCode == EOS_EResult::EOS_Success;
 		if (bWasSuccessful)
 		{
-			if (const TSharedRef<TArray<FOnlineAchievement>>* Achievements = CachedAchievementsMap.Find(LambdaPlayerId))
+			if (TArray<FOnlineAchievement>* Achievements = CachedAchievementsMap.Find(LambdaPlayerId))
 			{
-				UE_LOG_ONLINE_ACHIEVEMENTS(Verbose, TEXT("(%d) achievements unlocked. Caching achievements for user (%s)"), (uint32)Data->AchievementsCount, *LambdaPlayerId->ToString());
+				UE_LOG_ONLINE_ACHIEVEMENTS(Verbose, TEXT("(%d) achievements unlocked. Updating achievement cache for user (%s)"), (uint32)Data->AchievementsCount, *LambdaPlayerId->ToString());
 
 				for (const FString& AchievementId : InAchievementIds)
 				{
-					for (FOnlineAchievement& Achievement : Achievements->Get())
+					for (FOnlineAchievement& Achievement : *Achievements)
 					{
 						if (Achievement.Id == AchievementId)
 						{
@@ -78,8 +78,13 @@ void FOnlineAchievementsEOS::UnlockAchievements(const FUniqueNetId& PlayerId, FO
 			}
 			else
 			{
-				UE_LOG_ONLINE_ACHIEVEMENTS(Error, TEXT("Cached Achievements not found for local user (%s)"), *LambdaPlayerId->ToString());
-				bWasSuccessful = false;
+				// Achievements were never queried and cached prior to calling UnlockAchievements. Returning success and logging that the achievement cache is empty
+				UE_LOG_ONLINE_ACHIEVEMENTS(Verbose, TEXT("Achievements unlocked for user (%s). No cached achievements for this user, please call QueryAchievements to build the achievement cache."), *LambdaPlayerId->ToString());
+				
+				for (const FString& AchievementId : InAchievementIds)
+				{
+					TriggerOnAchievementUnlockedDelegates(*LambdaPlayerId, AchievementId);
+				}
 			}
 		}
 		else
@@ -117,8 +122,7 @@ void FOnlineAchievementsEOS::QueryAchievements(const FUniqueNetId& PlayerId, con
 		bool bWasSuccessful = Data->ResultCode == EOS_EResult::EOS_Success;
 		if (bWasSuccessful)
 		{
-			TSharedRef<TArray<FOnlineAchievement>> Cheevos = MakeShareable(new TArray<FOnlineAchievement>());
-			CachedAchievementsMap.Add(LambdaPlayerId, Cheevos);
+			TArray<FOnlineAchievement>& PlayerAchievements = CachedAchievementsMap.Add(LambdaPlayerId);
 
 			int32 LocalUserNum = EOSSubsystem->UserManager->GetLocalUserNumFromUniqueNetId(*LambdaPlayerId);
 			EOS_ProductUserId UserId = EOSSubsystem->UserManager->GetLocalProductUserId(LocalUserNum);
@@ -143,16 +147,15 @@ void FOnlineAchievementsEOS::QueryAchievements(const FUniqueNetId& PlayerId, con
 				EOS_EResult Result = EOS_Achievements_CopyPlayerAchievementByIndex(EOSSubsystem->AchievementsHandle, &CopyOptions, &AchievementEOS);
 				if (Result == EOS_EResult::EOS_Success)
 				{
-					FOnlineAchievement* Achievement = new(*Cheevos) FOnlineAchievement();
-
-					Achievement->Id = AchievementEOS->AchievementId;
-					Achievement->Progress = AchievementEOS->Progress;
+					FOnlineAchievement& Achievement = PlayerAchievements.Emplace_GetRef();
+					Achievement.Id = AchievementEOS->AchievementId;
+					Achievement.Progress = AchievementEOS->Progress;
 
 					EOS_Achievements_PlayerAchievement_Release(AchievementEOS);
 
 					if (UE_BUILD_DEBUG)
 					{
-						UE_LOG_ONLINE_ACHIEVEMENTS(Log, TEXT("Achievement progress (%s)"), *Achievement->ToDebugString());
+						UE_LOG_ONLINE_ACHIEVEMENTS(Log, TEXT("Achievement progress (%s)"), *Achievement.ToDebugString());
 					}
 				}
 				else
@@ -252,9 +255,9 @@ void FOnlineAchievementsEOS::QueryAchievementDescriptions(const FUniqueNetId& Pl
 
 EOnlineCachedResult::Type FOnlineAchievementsEOS::GetCachedAchievement(const FUniqueNetId& PlayerId, const FString& AchievementId, FOnlineAchievement& OutAchievement)
 {
-	if (const TSharedRef<TArray<FOnlineAchievement>>* Achievements = CachedAchievementsMap.Find(PlayerId.AsShared()))
+	if (const TArray<FOnlineAchievement>* Achievements = CachedAchievementsMap.Find(PlayerId.AsShared()))
 	{
-		for (const FOnlineAchievement& Achievement : **Achievements)
+		for (const FOnlineAchievement& Achievement : *Achievements)
 		{
 			if (Achievement.Id == AchievementId)
 			{
@@ -268,9 +271,9 @@ EOnlineCachedResult::Type FOnlineAchievementsEOS::GetCachedAchievement(const FUn
 
 EOnlineCachedResult::Type FOnlineAchievementsEOS::GetCachedAchievements(const FUniqueNetId& PlayerId, TArray<FOnlineAchievement>& OutAchievements)
 {
-	if (const TSharedRef<TArray<FOnlineAchievement>>* CachedAchievements = CachedAchievementsMap.Find(PlayerId.AsShared()))
+	if (const TArray<FOnlineAchievement>* CachedAchievements = CachedAchievementsMap.Find(PlayerId.AsShared()))
 	{
-		OutAchievements = **CachedAchievements;
+		OutAchievements = *CachedAchievements;
 		return EOnlineCachedResult::Success;
 	}
 	return EOnlineCachedResult::NotFound;
diff --git a/Engine/Plugins/Online/OnlineSubsystemEOS/Source/OnlineSubsystemEOS/Private/OnlineAchievementsEOS.h b/Engine/Plugins/Online/OnlineSubsystemEOS/Source/OnlineSubsystemEOS/Private/OnlineAchievementsEOS.h
--- a/Engine/Plugins/Online/OnlineSubsystemEOS/Source/OnlineSubsystemEOS/Private/OnlineAchievementsEOS.h
+++ b/Engine/Plugins/Online/OnlineSubsystemEOS/Source/OnlineSubsystemEOS/Private/OnlineAchievementsEOS.h
@@ -44,7 +44,7 @@ private:
 	/** Reference to the main EOS subsystem */
 	FOnlineSubsystemEOS* EOSSubsystem;
 	/** Holds the cached info from the last time this was called */
-	TUniqueNetIdMap<TSharedRef<TArray<FOnlineAchievement>>> CachedAchievementsMap;
+	TUniqueNetIdMap<TArray<FOnlineAchievement>> CachedAchievementsMap;
 	/** Holds the cached info from the last time this was called */
 	TArray<FOnlineAchievementDesc> CachedAchievementDefinitions;
 	/** Hack to work around the fact the description doesn't include the ID */
diff --git a/Engine/Plugins/Online/OnlineSubsystemEOS/Source/OnlineSubsystemEOS/Private/UserManagerEOS.cpp b/Engine/Plugins/Online/OnlineSubsystemEOS/Source/OnlineSubsystemEOS/Private/UserManagerEOS.cpp
--- a/Engine/Plugins/Online/OnlineSubsystemEOS/Source/OnlineSubsystemEOS/Private/UserManagerEOS.cpp
+++ b/Engine/Plugins/Online/OnlineSubsystemEOS/Source/OnlineSubsystemEOS/Private/UserManagerEOS.cpp
@@ -2438,11 +2438,21 @@ ELoginStatus::Type FUserManagerEOS::GetLoginStatus(const FUniqueNetIdEOS& UserId
 	if (!Settings.bUseEAS)
 	{
 		const EOS_ProductUserId ProductUserId = UserId.GetProductUserId();
-		if (ProductUserId != nullptr)
+		if (ProductUserId == nullptr)
 		{
-			return ELoginStatus::LoggedIn;
+			return ELoginStatus::NotLoggedIn;
+		}
+		EOS_ELoginStatus LoginStatus = EOS_Connect_GetLoginStatus(EOSSubsystem->ConnectHandle, ProductUserId);
+		
+		// EOS_ELoginStatus::EOS_LS_UsingLocalProfile is only used with EOS_Auth
+		if (LoginStatus == EOS_ELoginStatus::EOS_LS_LoggedIn)
+		{
+			return ELoginStatus::LoggedIn; 
+		}
+		else
+		{
+			return ELoginStatus::NotLoggedIn;
 		}
-		return ELoginStatus::NotLoggedIn;
 	}
 
 	const EOS_EpicAccountId AccountId = UserId.GetEpicAccountId();
