Warnings incorrectly thrown when creating/joining a lobby.

diff --git a/Engine/Plugins/Online/OnlineSubsystemEOS/Source/OnlineSubsystemEOS/Private/OnlineSessionEOS.cpp b/Engine/Plugins/Online/OnlineSubsystemEOS/Source/OnlineSubsystemEOS/Private/OnlineSessionEOS.cpp
--- a/Engine/Plugins/Online/OnlineSubsystemEOS/Source/OnlineSubsystemEOS/Private/OnlineSessionEOS.cpp
+++ b/Engine/Plugins/Online/OnlineSubsystemEOS/Source/OnlineSubsystemEOS/Private/OnlineSessionEOS.cpp
@@ -4473,7 +4473,7 @@ void FOnlineSessionEOS::SetLobbyMemberAttributes(EOS_HLobbyModification LobbyMod
 	}
 	else
 	{
-		UE_LOG_ONLINE_SESSION(Warning, TEXT("[FOnlineSessionEOS::SetLobbyMemberAttributes] Lobby Member with UniqueNetId [%s] not found in Lobby Session with SessionId [%s]"), *LobbyMemberId->ToDebugString(), *Session.GetSessionIdStr());
+		UE_LOG_ONLINE_SESSION(Verbose, TEXT("[FOnlineSessionEOS::SetLobbyMemberAttributes] Lobby Member with UniqueNetId [%s] not found in Lobby Session with SessionId [%s]. This is expected when the lobby owner creates the lobby, or when a new member joins."), *LobbyMemberId->ToDebugString(), *Session.GetSessionIdStr());
 	}
 }
 
