- Bind to OnConfigSectionsChanged and invalidate the cached settings used by ManualGetSettings.

diff --git a/Engine/Plugins/Online/OnlineSubsystemEOS/Source/OnlineSubsystemEOS/Private/EOSSettings.cpp b/Engine/Plugins/Online/OnlineSubsystemEOS/Source/OnlineSubsystemEOS/Private/EOSSettings.cpp
--- a/Engine/Plugins/Online/OnlineSubsystemEOS/Source/OnlineSubsystemEOS/Private/EOSSettings.cpp
+++ b/Engine/Plugins/Online/OnlineSubsystemEOS/Source/OnlineSubsystemEOS/Private/EOSSettings.cpp
@@ -7,6 +7,7 @@
 #include "Algo/Transform.h"
 #include "Misc/CommandLine.h"
 #include "Misc/ConfigCacheIni.h"
+#include "Misc/CoreDelegates.h"
 
 #include UE_INLINE_GENERATED_CPP_BY_NAME(EOSSettings)
 
@@ -18,40 +19,46 @@
 
 #define INI_SECTION TEXT("/Script/OnlineSubsystemEOS.EOSSettings")
 
-inline bool IsAnsi(const FString& Source)
+namespace
 {
-	for (const TCHAR& IterChar : Source)
+	inline bool IsAnsi(const FString& Source)
 	{
-		if (!FChar::IsPrint(IterChar))
+		for (const TCHAR& IterChar : Source)
 		{
-			return false;
+			if (!FChar::IsPrint(IterChar))
+			{
+				return false;
+			}
 		}
+		return true;
 	}
-	return true;
-}
 
-inline bool IsHex(const FString& Source)
-{
-	for (const TCHAR& IterChar : Source)
+	inline bool IsHex(const FString& Source)
 	{
-		if (!FChar::IsHexDigit(IterChar))
+		for (const TCHAR& IterChar : Source)
 		{
-			return false;
+			if (!FChar::IsHexDigit(IterChar))
+			{
+				return false;
+			}
 		}
+		return true;
 	}
-	return true;
-}
 
-inline bool ContainsWhitespace(const FString& Source)
-{
-	for (const TCHAR& IterChar : Source)
+	inline bool ContainsWhitespace(const FString& Source)
 	{
-		if (FChar::IsWhitespace(IterChar))
+		for (const TCHAR& IterChar : Source)
 		{
-			return true;
+			if (FChar::IsWhitespace(IterChar))
+			{
+				return true;
+			}
 		}
+		return false;
 	}
-	return false;
+
+	static TOptional<FEOSSettings> GCachedSettings;
+	static FDelegateHandle GOnConfigSectionsChangedDelegateHandle;
 }
 
 FEOSArtifactSettings FArtifactSettings::ToNative() const
@@ -168,43 +175,41 @@ FEOSSettings UEOSSettings::AutoGetSettings()
 
 const FEOSSettings& UEOSSettings::ManualGetSettings()
 {
-	static TOptional<FEOSSettings> CachedSettings;
-
-	if (!CachedSettings.IsSet())
+	if (!GCachedSettings.IsSet())
 	{
-		CachedSettings.Emplace();
+		GCachedSettings.Emplace();
 
-		GConfig->GetString(INI_SECTION, TEXT("CacheDir"), CachedSettings->CacheDir, GEngineIni);
-		GConfig->GetString(INI_SECTION, TEXT("DefaultArtifactName"), CachedSettings->DefaultArtifactName, GEngineIni);
-		GConfig->GetString(INI_SECTION, TEXT("SteamTokenType"), CachedSettings->SteamTokenType, GEngineIni);
-		CachedSettings->RTCBackgroundMode = EOS_ERTCBackgroundMode::EOS_RTCBM_KeepRoomsAlive;
+		GConfig->GetString(INI_SECTION, TEXT("CacheDir"), GCachedSettings->CacheDir, GEngineIni);
+		GConfig->GetString(INI_SECTION, TEXT("DefaultArtifactName"), GCachedSettings->DefaultArtifactName, GEngineIni);
+		GConfig->GetString(INI_SECTION, TEXT("SteamTokenType"), GCachedSettings->SteamTokenType, GEngineIni);
+		GCachedSettings->RTCBackgroundMode = EOS_ERTCBackgroundMode::EOS_RTCBM_KeepRoomsAlive;
 		FString RTCBackgroundModeStr;
 		GConfig->GetString(INI_SECTION, TEXT("RTCBackgroundMode"), RTCBackgroundModeStr, GEngineIni);
 		if (!RTCBackgroundModeStr.IsEmpty())
 		{
-			LexFromString(CachedSettings->RTCBackgroundMode, *RTCBackgroundModeStr);
+			LexFromString(GCachedSettings->RTCBackgroundMode, *RTCBackgroundModeStr);
 		}
-		GConfig->GetInt(INI_SECTION, TEXT("TickBudgetInMilliseconds"), CachedSettings->TickBudgetInMilliseconds, GEngineIni);
-		GConfig->GetInt(INI_SECTION, TEXT("TitleStorageReadChunkLength"), CachedSettings->TitleStorageReadChunkLength, GEngineIni);
-		GConfig->GetBool(INI_SECTION, TEXT("bEnableOverlay"), CachedSettings->bEnableOverlay, GEngineIni);
-		GConfig->GetBool(INI_SECTION, TEXT("bEnableSocialOverlay"), CachedSettings->bEnableSocialOverlay, GEngineIni);
-		GConfig->GetBool(INI_SECTION, TEXT("bEnableEditorOverlay"), CachedSettings->bEnableEditorOverlay, GEngineIni);
-		GConfig->GetBool(INI_SECTION, TEXT("bPreferPersistentAuth"), CachedSettings->bPreferPersistentAuth, GEngineIni);
-		GConfig->GetBool(INI_SECTION, TEXT("bUseEAS"), CachedSettings->bUseEAS, GEngineIni);
-		GConfig->GetBool(INI_SECTION, TEXT("bUseEOSConnect"), CachedSettings->bUseEOSConnect, GEngineIni);
-		GConfig->GetBool(INI_SECTION, TEXT("bUseEOSSessions"), CachedSettings->bUseEOSSessions, GEngineIni);
+		GConfig->GetInt(INI_SECTION, TEXT("TickBudgetInMilliseconds"), GCachedSettings->TickBudgetInMilliseconds, GEngineIni);
+		GConfig->GetInt(INI_SECTION, TEXT("TitleStorageReadChunkLength"), GCachedSettings->TitleStorageReadChunkLength, GEngineIni);
+		GConfig->GetBool(INI_SECTION, TEXT("bEnableOverlay"), GCachedSettings->bEnableOverlay, GEngineIni);
+		GConfig->GetBool(INI_SECTION, TEXT("bEnableSocialOverlay"), GCachedSettings->bEnableSocialOverlay, GEngineIni);
+		GConfig->GetBool(INI_SECTION, TEXT("bEnableEditorOverlay"), GCachedSettings->bEnableEditorOverlay, GEngineIni);
+		GConfig->GetBool(INI_SECTION, TEXT("bPreferPersistentAuth"), GCachedSettings->bPreferPersistentAuth, GEngineIni);
+		GConfig->GetBool(INI_SECTION, TEXT("bUseEAS"), GCachedSettings->bUseEAS, GEngineIni);
+		GConfig->GetBool(INI_SECTION, TEXT("bUseEOSConnect"), GCachedSettings->bUseEOSConnect, GEngineIni);
+		GConfig->GetBool(INI_SECTION, TEXT("bUseEOSSessions"), GCachedSettings->bUseEOSSessions, GEngineIni);
 PRAGMA_DISABLE_DEPRECATION_WARNINGS
-		GConfig->GetBool(INI_SECTION, TEXT("bMirrorStatsToEOS"), CachedSettings->bMirrorStatsToEOS, GEngineIni);
-		GConfig->GetBool(INI_SECTION, TEXT("bMirrorAchievementsToEOS"), CachedSettings->bMirrorAchievementsToEOS, GEngineIni);
-		GConfig->GetBool(INI_SECTION, TEXT("bMirrorPresenceToEAS"), CachedSettings->bMirrorPresenceToEAS, GEngineIni);
+		GConfig->GetBool(INI_SECTION, TEXT("bMirrorStatsToEOS"), GCachedSettings->bMirrorStatsToEOS, GEngineIni);
+		GConfig->GetBool(INI_SECTION, TEXT("bMirrorAchievementsToEOS"), GCachedSettings->bMirrorAchievementsToEOS, GEngineIni);
+		GConfig->GetBool(INI_SECTION, TEXT("bMirrorPresenceToEAS"), GCachedSettings->bMirrorPresenceToEAS, GEngineIni);
 PRAGMA_ENABLE_DEPRECATION_WARNINGS
-		GConfig->GetBool(INI_SECTION, TEXT("bUseNewLoginFlow"), CachedSettings->bUseNewLoginFlow, GEngineIni);
+		GConfig->GetBool(INI_SECTION, TEXT("bUseNewLoginFlow"), GCachedSettings->bUseNewLoginFlow, GEngineIni);
 		// Artifacts explicitly skipped
-		GConfig->GetArray(INI_SECTION, TEXT("TitleStorageTags"), CachedSettings->TitleStorageTags, GEngineIni);
-		GConfig->GetArray(INI_SECTION, TEXT("AuthScopeFlags"), CachedSettings->AuthScopeFlags, GEngineIni);
+		GConfig->GetArray(INI_SECTION, TEXT("TitleStorageTags"), GCachedSettings->TitleStorageTags, GEngineIni);
+		GConfig->GetArray(INI_SECTION, TEXT("AuthScopeFlags"), GCachedSettings->AuthScopeFlags, GEngineIni);
 	}
 
-	return *CachedSettings;
+	return *GCachedSettings;
 }
 
 FEOSSettings UEOSSettings::ToNative() const
@@ -532,5 +537,24 @@ PRAGMA_ENABLE_DEPRECATION_WARNINGS
 
 #endif
 
+void UEOSSettings::ModuleInit()
+{
+	check(!GOnConfigSectionsChangedDelegateHandle.IsValid());
+	GOnConfigSectionsChangedDelegateHandle = FCoreDelegates::TSOnConfigSectionsChanged().AddLambda([](const FString& IniFilename, const TSet<FString>& SectionNames)
+		{
+			if (IniFilename == GEngineIni && SectionNames.Contains(INI_SECTION))
+			{
+				GCachedSettings.Reset();
+			}
+		});
+}
+
+void UEOSSettings::ModuleShutdown()
+{
+	FCoreDelegates::TSOnConfigSectionsChanged().Remove(GOnConfigSectionsChangedDelegateHandle);
+	GOnConfigSectionsChangedDelegateHandle.Reset();
+	GCachedSettings.Reset();
+}
+
 #undef LOCTEXT_NAMESPACE
 
diff --git a/Engine/Plugins/Online/OnlineSubsystemEOS/Source/OnlineSubsystemEOS/Private/OnlineSubsystemEOSModule.cpp b/Engine/Plugins/Online/OnlineSubsystemEOS/Source/OnlineSubsystemEOS/Private/OnlineSubsystemEOSModule.cpp
--- a/Engine/Plugins/Online/OnlineSubsystemEOS/Source/OnlineSubsystemEOS/Private/OnlineSubsystemEOSModule.cpp
+++ b/Engine/Plugins/Online/OnlineSubsystemEOS/Source/OnlineSubsystemEOS/Private/OnlineSubsystemEOSModule.cpp
@@ -101,6 +101,7 @@ void FOnlineSubsystemEOSModule::StartupModule()
 #if WITH_EOS_SDK
 	// Have to call this as early as possible in order to hook the rendering device
 	FOnlineSubsystemEOS::ModuleInit();
+	UEOSSettings::ModuleInit();
 #endif
 
 #if WITH_EDITOR
@@ -141,6 +142,7 @@ void FOnlineSubsystemEOSModule::ShutdownModule()
 
 #if WITH_EOS_SDK
 	FOnlineSubsystemEOS::ModuleShutdown();
+	UEOSSettings::ModuleShutdown();
 #endif
 
 	FOnlineSubsystemModule& OSS = FModuleManager::GetModuleChecked<FOnlineSubsystemModule>("OnlineSubsystem");
diff --git a/Engine/Plugins/Online/OnlineSubsystemEOS/Source/OnlineSubsystemEOS/Public/EOSSettings.h b/Engine/Plugins/Online/OnlineSubsystemEOS/Source/OnlineSubsystemEOS/Public/EOSSettings.h
--- a/Engine/Plugins/Online/OnlineSubsystemEOS/Source/OnlineSubsystemEOS/Public/EOSSettings.h
+++ b/Engine/Plugins/Online/OnlineSubsystemEOS/Source/OnlineSubsystemEOS/Public/EOSSettings.h
@@ -235,4 +235,8 @@ private:
 
 	static FEOSSettings AutoGetSettings();
 	static const FEOSSettings& ManualGetSettings();
+
+	friend class FOnlineSubsystemEOSModule;
+	static void ModuleInit();
+	static void ModuleShutdown();
 };
\ No newline at end of file
