Addition of support for EOS_Connect_AddNotifyAuthExpiration and related changes

diff --git a/Engine/Plugins/Online/OnlineSubsystem/Source/Private/OnlineSubsystemImpl.cpp b/Engine/Plugins/Online/OnlineSubsystem/Source/Private/OnlineSubsystemImpl.cpp
--- a/Engine/Plugins/Online/OnlineSubsystem/Source/Private/OnlineSubsystemImpl.cpp
+++ b/Engine/Plugins/Online/OnlineSubsystem/Source/Private/OnlineSubsystemImpl.cpp
@@ -252,6 +252,10 @@ bool FOnlineSubsystemImpl::Exec(UWorld* InWorld, const TCHAR* Cmd, FOutputDevice
 	{
 		bWasHandled = HandleFriendExecCommands(InWorld, Cmd, Ar);
 	}
+	else if (FParse::Command(&Cmd, TEXT("IDENTITY")))
+	{
+		bWasHandled = HandleIdentityExecCommands(InWorld, Cmd, Ar);
+	}
 	else if (FParse::Command(&Cmd, TEXT("SESSION")))
 	{
 		bWasHandled = HandleSessionExecCommands(InWorld, Cmd, Ar);
@@ -696,6 +700,27 @@ bool FOnlineSubsystemImpl::HandleFriendExecCommands(UWorld* InWorld, const TCHAR
 	return bWasHandled;
 }
 
+bool FOnlineSubsystemImpl::HandleIdentityExecCommands(UWorld* InWorld, const TCHAR* Cmd, FOutputDevice& Ar)
+{
+	bool bWasHandled = false;
+
+	if (FParse::Command(&Cmd, TEXT("TRIGGERONAUTHABOUTTOEXPIREDELEGATES"))) /* ONLINE IDENTITY TRIGGERONAUTHABOUTTOEXPIREDELEGATES LocalUserNum=0 */
+	{
+		IOnlineIdentityPtr IdentityInt = GetIdentityInterface();
+		if (IdentityInt.IsValid())
+		{
+			int LocalUserNum = 0;
+			FParse::Value(Cmd, TEXT("LocalUserNum="), LocalUserNum);
+
+			IdentityInt->TriggerOnAuthAboutToExpireDelegates(LocalUserNum);
+		}
+		
+		bWasHandled = true;
+	}
+
+	return bWasHandled;
+}
+
 bool FOnlineSubsystemImpl::HandlePresenceExecCommands(UWorld* InWorld, const TCHAR* Cmd, FOutputDevice& Ar)
 {
 	bool bWasHandled = false;
diff --git a/Engine/Plugins/Online/OnlineSubsystem/Source/Public/Interfaces/OnlineIdentityInterface.h b/Engine/Plugins/Online/OnlineSubsystem/Source/Public/Interfaces/OnlineIdentityInterface.h
--- a/Engine/Plugins/Online/OnlineSubsystem/Source/Public/Interfaces/OnlineIdentityInterface.h
+++ b/Engine/Plugins/Online/OnlineSubsystem/Source/Public/Interfaces/OnlineIdentityInterface.h
@@ -139,6 +139,15 @@ typedef FOnLoginChanged::FDelegate FOnLoginChangedDelegate;
 DECLARE_MULTICAST_DELEGATE_FourParams(FOnLoginStatusChanged, int32 /*LocalUserNum*/, ELoginStatus::Type /*OldStatus*/, ELoginStatus::Type /*NewStatus*/, const FUniqueNetId& /*NewId*/);
 typedef FOnLoginStatusChanged::FDelegate FOnLoginStatusChangedDelegate;
 
+/**
+ * Delegate called when an auth session is about to expire and a new auth token is needed.
+ * The code bound to it is expected to respond by getting a new auth token and use it to re-call Login to refresh auth and keep the player logged in.
+ * 
+ * param LocalUserNum the controller number of the user whose auth is about to expire
+ */
+DECLARE_MULTICAST_DELEGATE_OneParam(FOnAuthAboutToExpire, int32 /*LocalUserNum*/);
+typedef FOnAuthAboutToExpire::FDelegate FOnAuthAboutToExpireDelegate;
+
 /**
  * Delegate called when a controller-user pairing changes
  *
@@ -268,6 +277,14 @@ public:
 	 */
 	DEFINE_ONLINE_PLAYER_DELEGATE_THREE_PARAM(MAX_LOCAL_PLAYERS, OnLoginStatusChanged, ELoginStatus::Type /*OldStatus*/, ELoginStatus::Type /*NewStatus*/, const FUniqueNetId& /*NewId*/);
 
+	/**
+	 * Delegate called when an auth session is about to expire and a new auth token is needed.
+	 * The code bound to it is expected to respond by getting a new auth token and use it to re-call Login to refresh auth and keep the player logged in.
+	 * 
+	 * param LocalUserNum the controller number of the user whose auth is about to expire
+	 */
+	DEFINE_ONLINE_DELEGATE_ONE_PARAM(OnAuthAboutToExpire, int32 /*LocalUserNum*/);
+
 	/**
 	 * Delegate called when a controller-user pairing changes
 	 *
diff --git a/Engine/Plugins/Online/OnlineSubsystem/Source/Public/OnlineSubsystemImpl.h b/Engine/Plugins/Online/OnlineSubsystem/Source/Public/OnlineSubsystemImpl.h
--- a/Engine/Plugins/Online/OnlineSubsystem/Source/Public/OnlineSubsystemImpl.h
+++ b/Engine/Plugins/Online/OnlineSubsystem/Source/Public/OnlineSubsystemImpl.h
@@ -29,6 +29,7 @@ private:
 	 * Exec function handling for Exec() call
 	 */
 	bool HandleFriendExecCommands(UWorld* InWorld, const TCHAR* Cmd, FOutputDevice& Ar);
+	bool HandleIdentityExecCommands(UWorld* InWorld, const TCHAR* Cmd, FOutputDevice& Ar);
 	bool HandleSessionExecCommands(UWorld* InWorld, const TCHAR* Cmd, FOutputDevice& Ar);
 	bool HandlePresenceExecCommands(UWorld* InWorld, const TCHAR* Cmd, FOutputDevice& Ar);
 	bool HandlePurchaseExecCommands(UWorld* InWorld, const TCHAR* Cmd, FOutputDevice& Ar);
diff --git a/Engine/Plugins/Online/OnlineSubsystemEOS/Source/OnlineSubsystemEOS/Private/UserManagerEOS.cpp b/Engine/Plugins/Online/OnlineSubsystemEOS/Source/OnlineSubsystemEOS/Private/UserManagerEOS.cpp
--- a/Engine/Plugins/Online/OnlineSubsystemEOS/Source/OnlineSubsystemEOS/Private/UserManagerEOS.cpp
+++ b/Engine/Plugins/Online/OnlineSubsystemEOS/Source/OnlineSubsystemEOS/Private/UserManagerEOS.cpp
@@ -533,15 +533,31 @@ bool FUserManagerEOS::Login(int32 LocalUserNum, const FOnlineAccountCredentials&
 		return LoginLegacy(LocalUserNum, AccountCredentials);
 	}
 
-
-	// Connect Login without EAS is not supported, as we don't yet have a way to prompt the user to provide a token for auth refresh
-	if (!Settings.bUseEAS)
+	if (!Settings.bUseEAS && !Settings.bUseEOSConnect)
 	{
-		UE_LOG_ONLINE(Warning, TEXT("EAS is not configured to be used. Failed to login in user (%d)"), LocalUserNum);
+		UE_LOG_ONLINE(Warning, TEXT("Neither EAS or EOS are configured to be used. Failed to login in user (%d)"), LocalUserNum);
 		TriggerOnLoginCompleteDelegates(LocalUserNum, false, *FUniqueNetIdEOS::EmptyId(), FString(TEXT("Not configured")));
 		return true;
 	}
 
+	if(!LoginSessions.IsValidIndex(LocalUserNum))
+	{
+		LoginSessions.Insert(LocalUserNum, FLoginSession());
+	}
+
+	// EOS Connect without EAS
+	if(!Settings.bUseEAS && Settings.bUseEOSConnect)
+	{
+		if(!AccountCredentials.Token.IsEmpty())
+		{
+			// We'll save that a Connect Token was manually provided for this local user
+			LoginSessions[LocalUserNum].UserProvidedConnectAuthToken = AccountCredentials.Token;
+		}
+
+		CallEOSConnectLogin(LocalUserNum, AccountCredentials);
+		return true;
+	}
+
 	// We don't support offline logged in, so they are either logged in or not
 	if (GetLoginStatus(LocalUserNum) == ELoginStatus::LoggedIn)
 	{
@@ -550,6 +566,7 @@ bool FUserManagerEOS::Login(int32 LocalUserNum, const FOnlineAccountCredentials&
 		return true;
 	}
 
+	// EAS Auth Login + EOS Connect Login
 	CallEOSAuthLogin(LocalUserNum, AccountCredentials, false);
 	return true;
 }
@@ -731,6 +748,7 @@ void FUserManagerEOS::LoginViaExternalAuth(int32 LocalUserNum)
 					if (!bWasSuccessful || !AuthToken.IsValid())
 					{
 						UE_LOG_ONLINE(Warning, TEXT("Unable to Login() user (%d) due to an empty platform auth token"), LocalUserNum);
+						LoginSessions.RemoveAt(LocalUserNum);
 						TriggerOnLoginCompleteDelegates(LocalUserNum, false, *FUniqueNetIdEOS::EmptyId(), FString(TEXT("Missing platform auth token")));
 						return;
 					}
@@ -749,6 +767,7 @@ void FUserManagerEOS::LoginViaExternalAuth(int32 LocalUserNum)
 					else
 					{
 						UE_LOG_ONLINE(Error, TEXT("FAuthCredentials object cannot be constructed with invalid FExternalAuthToken parameter"));
+						LoginSessions.RemoveAt(LocalUserNum);
 						TriggerOnLoginCompleteDelegates(LocalUserNum, false, *FUniqueNetIdEOS::EmptyId(), FString(TEXT("Invalid platform auth token")));
 						return;
 					}
@@ -808,6 +827,7 @@ void FUserManagerEOS::CallEOSAuthLogin(int32 LocalUserNum, const FOnlineAccountC
 	if (!ToEOS_ELoginCredentialType(CredentialsType, EOSCredentials.Type))
 	{
 		UE_LOG_ONLINE(Warning, TEXT("Unable to Login() user (%d) due to missing auth parameters"), LocalUserNum);
+		LoginSessions.RemoveAt(LocalUserNum);
 		TriggerOnLoginCompleteDelegates(LocalUserNum, false, *FUniqueNetIdEOS::EmptyId(), FString(TEXT("Missing auth parameters")));
 		return;
 	}
@@ -831,6 +851,7 @@ void FUserManagerEOS::CallEOSAuthLogin(int32 LocalUserNum, const FOnlineAccountC
 				if (!LexFromString(EOSCredentials.ExternalType, *CredentialsTokenType))
 				{
 					UE_LOG_ONLINE(Warning, TEXT("Unable to Login() user (%d). External Auth Token Type not valid."), LocalUserNum);
+					LoginSessions.RemoveAt(LocalUserNum);
 					TriggerOnLoginCompleteDelegates(LocalUserNum, false, *FUniqueNetIdEOS::EmptyId(), FString(TEXT("External Auth Token Type not valid")));
 					return;
 				}
@@ -838,6 +859,7 @@ void FUserManagerEOS::CallEOSAuthLogin(int32 LocalUserNum, const FOnlineAccountC
 			else
 			{
 				UE_LOG_ONLINE(Warning, TEXT("Unable to Login() user (%d). External Auth Token Type not specified."), LocalUserNum);
+				LoginSessions.RemoveAt(LocalUserNum);
 				TriggerOnLoginCompleteDelegates(LocalUserNum, false, *FUniqueNetIdEOS::EmptyId(), FString(TEXT("External Auth Token Type not specified")));
 				return;
 			}
@@ -1042,6 +1064,7 @@ void FUserManagerEOS::OnEOSAuthLoginComplete(int32 LocalUserNum, const EOS_ELogi
 		{
 			FString ErrorString = FString::Printf(TEXT("Login(%d) failed with EOS result code (%s)"), LocalUserNum, *LexToString(LoginResultCode));
 			UE_LOG_ONLINE(Warning, TEXT("%s"), *ErrorString);
+			LoginSessions.RemoveAt(LocalUserNum);
 			TriggerOnLoginCompleteDelegates(LocalUserNum, false, *FUniqueNetIdEOS::EmptyId(), ErrorString);
 		};
 
@@ -1137,6 +1160,7 @@ void FUserManagerEOS::LinkEAS(int32 LocalUserNum, EOS_ContinuanceToken Token)
 			{
 				FString ErrorString = FString::Printf(TEXT("Login(%d) failed with EOS result code (%s)"), LocalUserNum, *LexToString(Data->ResultCode));
 				UE_LOG_ONLINE(Warning, TEXT("%s"), *ErrorString);
+				LoginSessions.RemoveAt(LocalUserNum);
 				TriggerOnLoginCompleteDelegates(LocalUserNum, false, *FUniqueNetIdEOS::EmptyId(), ErrorString);
 			}
 		};
@@ -1146,46 +1170,121 @@ void FUserManagerEOS::LinkEAS(int32 LocalUserNum, EOS_ContinuanceToken Token)
 struct FConnectCredentials :
 	public EOS_Connect_Credentials
 {
-	FConnectCredentials(EOS_EExternalCredentialType InType, const FExternalAuthToken& AuthToken) :
+	FConnectCredentials(EOS_EExternalCredentialType InType, const FString& TokenString) :
 		EOS_Connect_Credentials()
 	{
-		if (AuthToken.HasTokenData())
-		{
-			Init(InType, AuthToken.TokenData);
-		}
-		else if (AuthToken.HasTokenString())
+		ApiVersion = 1;
+		UE_EOS_CHECK_API_MISMATCH(EOS_CONNECT_CREDENTIALS_API_LATEST, 1);
+		Token = TokenUtf8;
+		Type = InType;
+
+		FCStringAnsi::Strncpy(TokenUtf8, TCHAR_TO_UTF8(*TokenString), TokenString.Len() + 1);
+	}
+
+	char TokenUtf8[EOS_MAX_TOKEN_SIZE];
+};
+
+FString ConvertExternalAuthTokenToFString(const FExternalAuthToken& AuthToken)
+{
+	FString Result;
+
+	if (AuthToken.HasTokenData())
+	{
+		char* TokenAnsi = nullptr;
+		uint32_t InOutBufferLength = EOS_MAX_TOKEN_SIZE;
+		EOS_ByteArray_ToString(AuthToken.TokenData.GetData(), AuthToken.TokenData.Num(), TokenAnsi, &InOutBufferLength);
+
+		Result = FString(TokenAnsi);
+	}
+	else if (AuthToken.HasTokenString())
+	{
+		Result = AuthToken.TokenString;
+	}
+	else
+	{
+		UE_LOG_ONLINE(Error, TEXT("FAuthTokenConverter object cannot be constructed with invalid FExternalAuthToken parameter"));
+	}
+
+	return Result;
+}
+
+void FUserManagerEOS::CallEOSConnectLogin(int32 LocalUserNum, const FOnlineAccountCredentials& Credentials)
+{
+	// Now login into our EOS account
+	TArray<FString> CredentialsTypeStrParams;
+	Credentials.Type.ParseIntoArray(CredentialsTypeStrParams, TEXT(":"));
+
+	EOS_EExternalCredentialType ExternalCredentialType;
+
+	FEOSSettings Settings = UEOSSettings::GetSettings();
+	if (Settings.bUseNewLoginFlow)
+	{
+		if (CredentialsTypeStrParams.IsValidIndex(1))
 		{
-			Init(InType, AuthToken.TokenString);
+			FString CredentialsTokenType = CredentialsTypeStrParams[1];
+			if (!LexFromString(ExternalCredentialType, *CredentialsTokenType))
+			{
+				UE_LOG_ONLINE(Warning, TEXT("Unable to Login() user (%d). External Auth Token Type not valid."), LocalUserNum);
+				LoginSessions.RemoveAt(LocalUserNum);
+				TriggerOnLoginCompleteDelegates(LocalUserNum, false, *FUniqueNetIdEOS::EmptyId(), FString(TEXT("External Auth Token Type not valid")));
+				return;
+			}
 		}
 		else
 		{
-			UE_LOG_ONLINE(Error, TEXT("FConnectCredentials object cannot be constructed with invalid FExternalAuthToken parameter"));
+			UE_LOG_ONLINE(Warning, TEXT("Unable to Login() user (%d). External Auth Token Type not specified."), LocalUserNum);
+			LoginSessions.RemoveAt(LocalUserNum);
+			TriggerOnLoginCompleteDelegates(LocalUserNum, false, *FUniqueNetIdEOS::EmptyId(), FString(TEXT("External Auth Token Type not specified")));
+			return;
 		}
 	}
-
-	void Init(EOS_EExternalCredentialType InType, const FString& InTokenString)
+	else
 	{
-		ApiVersion = 1;
-		UE_EOS_CHECK_API_MISMATCH(EOS_CONNECT_CREDENTIALS_API_LATEST, 1);
-		Token = TokenAnsi;
-		Type = InType;
-
-		FCStringAnsi::Strncpy(TokenAnsi, TCHAR_TO_UTF8(*InTokenString), InTokenString.Len() + 1);
+		IOnlineSubsystem* PlatformOSS = GetPlatformOSS();
+		ExternalCredentialType = ToEOS_EExternalCredentialTypeLegacy(PlatformOSS ? PlatformOSS->GetSubsystemName() : EOSSubsystem->GetSubsystemName(), Credentials);
 	}
 
-	void Init(EOS_EExternalCredentialType InType, const TArray<uint8>& InToken)
-	{
-		ApiVersion = 1;
-		UE_EOS_CHECK_API_MISMATCH(EOS_CONNECT_CREDENTIALS_API_LATEST, 1);
-		Token = TokenAnsi;
-		Type = InType;
+	FConnectCredentials ConnectCredentials(ExternalCredentialType, Credentials.Token);
 
-		uint32_t InOutBufferLength = EOS_MAX_TOKEN_SIZE;
-		EOS_ByteArray_ToString(InToken.GetData(), InToken.Num(), TokenAnsi, &InOutBufferLength);
-	}
+	EOS_Connect_LoginOptions Options = { };
+	Options.ApiVersion = 2;
+	UE_EOS_CHECK_API_MISMATCH(EOS_CONNECT_LOGIN_API_LATEST, 2);
+	Options.Credentials = &ConnectCredentials;
 
-	char TokenAnsi[EOS_MAX_TOKEN_SIZE];
-};
+#if ADD_USER_LOGIN_INFO
+	EOS_Connect_UserLoginInfo UserLoginInfo = {};
+	UserLoginInfo.ApiVersion = 2;
+	UE_EOS_CHECK_API_MISMATCH(EOS_CONNECT_USERLOGININFO_API_LATEST, 2);
+	const FTCHARToUTF8 DisplayNameUtf8(*GetPlatformDisplayName(LocalUserNum));
+	UserLoginInfo.DisplayName = DisplayNameUtf8.Get();
+	UserLoginInfo.NsaIdToken = nullptr; // Not set when EAS is not used
+
+	Options.UserLoginInfo = &UserLoginInfo;
+#endif
+
+	FConnectLoginCallback* CallbackObj = new FConnectLoginCallback(AsWeak());
+	CallbackObj->CallbackLambda = [this, LocalUserNum](const EOS_Connect_LoginCallbackInfo* Data)
+		{
+			if (Data->ResultCode == EOS_EResult::EOS_Success)
+			{
+				// We have an account mapping to the platform account, skip to final login
+				FullLoginCallback(LocalUserNum, nullptr, Data->LocalUserId);
+			}
+			else if (Data->ResultCode == EOS_EResult::EOS_InvalidUser)
+			{
+				// We need to create the platform account mapping for this user using the continuation token
+				CreateConnectedLogin(LocalUserNum, nullptr, Data->ContinuanceToken);
+			}
+			else
+			{
+				const FString ErrorString = FString::Printf(TEXT("ConnectLoginNoEAS(%d) failed with EOS result code (%s)"), LocalUserNum, *LexToString(Data->ResultCode));
+				UE_LOG_ONLINE(Warning, TEXT("%s"), *ErrorString);
+				LoginSessions.RemoveAt(LocalUserNum);
+				TriggerOnLoginCompleteDelegates(LocalUserNum, false, *FUniqueNetIdEOS::EmptyId(), ErrorString);
+			}
+		};
+	EOS_Connect_Login(EOSSubsystem->ConnectHandle, &Options, CallbackObj, CallbackObj->GetCallbackPtr());
+}
 
 bool FUserManagerEOS::ConnectLoginNoEASLegacy(int32 LocalUserNum, const FOnlineAccountCredentials& AccountCredentials)
 {
@@ -1203,46 +1302,10 @@ bool FUserManagerEOS::ConnectLoginNoEASLegacy(int32 LocalUserNum, const FOnlineA
 					}
 
 					// Now login into our EOS account
-					FConnectCredentials Credentials(ToEOS_EExternalCredentialTypeLegacy(GetPlatformOSS()->GetSubsystemName(), AccountCredentials), AuthToken);
-					EOS_Connect_LoginOptions Options = { };
-					Options.ApiVersion = 2;
-					UE_EOS_CHECK_API_MISMATCH(EOS_CONNECT_LOGIN_API_LATEST, 2);
-					Options.Credentials = &Credentials;
-
-#if ADD_USER_LOGIN_INFO
-					EOS_Connect_UserLoginInfo UserLoginInfo = {};
-					UserLoginInfo.ApiVersion = 2;
-					UE_EOS_CHECK_API_MISMATCH(EOS_CONNECT_USERLOGININFO_API_LATEST, 2);
-					const FTCHARToUTF8 DisplayNameUtf8(*GetPlatformDisplayName(LocalUserNum));
-					UserLoginInfo.DisplayName = DisplayNameUtf8.Get();
-					UserLoginInfo.NsaIdToken = nullptr; // Not set when EAS is not used
-
-					Options.UserLoginInfo = &UserLoginInfo;
-#endif
-
-					FConnectLoginCallback* CallbackObj = new FConnectLoginCallback(AsWeak());
-					CallbackObj->CallbackLambda = [this, LocalUserNum, AccountCredentials](const EOS_Connect_LoginCallbackInfo* Data)
-						{
-							if (Data->ResultCode == EOS_EResult::EOS_Success)
-							{
-								// We have an account mapping to the platform account, skip to final login
-								FullLoginCallbackLegacy(LocalUserNum, nullptr, Data->LocalUserId, AccountCredentials);
-							}
-							else if (Data->ResultCode == EOS_EResult::EOS_InvalidUser)
-							{
-								// We need to create the platform account mapping for this user using the continuation token
-								CreateConnectedLoginLegacy(LocalUserNum, nullptr, Data->ContinuanceToken, AccountCredentials);
-							}
-							else
-							{
-								const FString ErrorString = FString::Printf(TEXT("ConnectLoginNoEAS(%d) failed with EOS result code (%s)"), LocalUserNum, *LexToString(Data->ResultCode));
-								UE_LOG_ONLINE(Warning, TEXT("%s"), *ErrorString);
-								TriggerOnLoginCompleteDelegates(LocalUserNum, false, *FUniqueNetIdEOS::EmptyId(), ErrorString);
-							}
-						};
-					EOS_Connect_Login(EOSSubsystem->ConnectHandle, &Options, CallbackObj, CallbackObj->GetCallbackPtr());
-			}
-}));
+					FString ExternalCredentialTypeStr = FString(TEXT("externalauth:")) + FString(LexToString(ToEOS_EExternalCredentialTypeLegacy(GetPlatformOSS()->GetSubsystemName(), AccountCredentials)));
+					CallEOSConnectLogin(LocalUserNum, FOnlineAccountCredentials(ExternalCredentialTypeStr, FString(), ConvertExternalAuthTokenToFString(AuthToken)));
+				}
+			}));
 
 	return true;
 }
@@ -1258,51 +1321,16 @@ bool FUserManagerEOS::ConnectLoginNoEAS(int32 LocalUserNum)
 					{
 						const FString ErrorString = FString::Printf(TEXT("ConnectLoginNoEAS(%d) failed due to the platform OSS giving an empty auth token"), LocalUserNum);
 						UE_LOG_ONLINE(Warning, TEXT("%s"), *ErrorString);
+						LoginSessions.RemoveAt(LocalUserNum);
 						TriggerOnLoginCompleteDelegates(LocalUserNum, false, *FUniqueNetIdEOS::EmptyId(), ErrorString);
 						return;
 					}
 
 					// Now login into our EOS account
-					FConnectCredentials Credentials(ToEOS_EExternalCredentialType(GetPlatformOSS()->GetSubsystemName()), AuthToken);
-					EOS_Connect_LoginOptions Options = { };
-					Options.ApiVersion = 2;
-					UE_EOS_CHECK_API_MISMATCH(EOS_CONNECT_LOGIN_API_LATEST, 2);
-					Options.Credentials = &Credentials;
-
-#if ADD_USER_LOGIN_INFO
-					EOS_Connect_UserLoginInfo UserLoginInfo = {};
-					UserLoginInfo.ApiVersion = 2;
-					UE_EOS_CHECK_API_MISMATCH(EOS_CONNECT_USERLOGININFO_API_LATEST, 2);
-					const FTCHARToUTF8 DisplayNameUtf8(*GetPlatformDisplayName(LocalUserNum));
-					UserLoginInfo.DisplayName = DisplayNameUtf8.Get();
-					UserLoginInfo.NsaIdToken = nullptr; // Not set when EAS is not used
-
-					Options.UserLoginInfo = &UserLoginInfo;
-#endif
-
-					FConnectLoginCallback* CallbackObj = new FConnectLoginCallback(AsWeak());
-					CallbackObj->CallbackLambda = [this, LocalUserNum](const EOS_Connect_LoginCallbackInfo* Data)
-						{
-							if (Data->ResultCode == EOS_EResult::EOS_Success)
-							{
-								// We have an account mapping to the platform account, skip to final login
-								FullLoginCallback(LocalUserNum, nullptr, Data->LocalUserId);
-							}
-							else if (Data->ResultCode == EOS_EResult::EOS_InvalidUser)
-							{
-								// We need to create the platform account mapping for this user using the continuation token
-								CreateConnectedLogin(LocalUserNum, nullptr, Data->ContinuanceToken);
-							}
-							else
-							{
-								const FString ErrorString = FString::Printf(TEXT("ConnectLoginNoEAS(%d) failed with EOS result code (%s)"), LocalUserNum, *LexToString(Data->ResultCode));
-								UE_LOG_ONLINE(Warning, TEXT("%s"), *ErrorString);
-								TriggerOnLoginCompleteDelegates(LocalUserNum, false, *FUniqueNetIdEOS::EmptyId(), ErrorString);
-							}
-						};
-					EOS_Connect_Login(EOSSubsystem->ConnectHandle, &Options, CallbackObj, CallbackObj->GetCallbackPtr());
-}
-}));
+					FString ExternalCredentialTypeStr = FString(TEXT("externalauth:")) + FString(LexToString(ToEOS_EExternalCredentialType(GetPlatformOSS()->GetSubsystemName())));
+					CallEOSConnectLogin(LocalUserNum, FOnlineAccountCredentials(ExternalCredentialTypeStr, FString(), ConvertExternalAuthTokenToFString(AuthToken)));	
+				}
+			}));
 
 	return true;
 }
@@ -1404,6 +1432,7 @@ bool FUserManagerEOS::ConnectLoginEAS(int32 LocalUserNum, EOS_EpicAccountId Acco
 					{
 						const FString ErrorString = FString::Printf(TEXT("ConnectLoginEAS(%d) failed due to the platform OSS giving an empty auth token"), LocalUserNum);
 						UE_LOG_ONLINE(Warning, TEXT("%s"), *ErrorString);
+						LoginSessions.RemoveAt(LocalUserNum);
 						TriggerOnLoginCompleteDelegates(LocalUserNum, false, *FUniqueNetIdEOS::EmptyId(), ErrorString);
 						return;
 					}
@@ -1477,7 +1506,7 @@ bool FUserManagerEOS::ConnectLoginEAS(int32 LocalUserNum, EOS_EpicAccountId Acco
 	return true;
 }
 
-void FUserManagerEOS::RefreshConnectLogin(int32 LocalUserNum)
+void FUserManagerEOS::AutoRefreshConnectLogin(int32 LocalUserNum)
 {
 	FString DefaultOSS;
 	GConfig->GetString(TEXT("OnlineSubsystem"), TEXT("DefaultPlatformService"), DefaultOSS, GEngineIni);
@@ -1527,53 +1556,86 @@ void FUserManagerEOS::RefreshConnectLogin(int32 LocalUserNum)
 		}
 		else
 		{
-			UE_LOG_ONLINE(Error, TEXT("[FUserManagerEOS::RefreshConnectLogin] AccessToken for user [%d] is empty"), LocalUserNum);
+			UE_LOG_ONLINE(Error, TEXT("[FUserManagerEOS::AutoRefreshConnectLogin] AccessToken for user [%d] is empty"), LocalUserNum);
 			Logout(LocalUserNum);
 		}
 	}
 	else
 	{
-		// Not using EAS so grab the platform auth token
-		GetPlatformAuthToken(LocalUserNum,
-			FOnGetLinkedAccountAuthTokenCompleteDelegate::CreateLambda([this, WeakThis = AsWeak()](int32 LocalUserNum, bool bWasSuccessful, const FExternalAuthToken& AuthToken)
-			{
-				if (FUserManagerEOSPtr StrongThis = WeakThis.Pin())
-				{
-					if (!bWasSuccessful || !AuthToken.IsValid())
-					{
-						UE_LOG_ONLINE(Error, TEXT("RefreshConnectLogin(%d) failed due to the platform OSS giving an empty auth token"), LocalUserNum);
-						Logout(LocalUserNum);
-						return;
-					}
-
-					// Now login into our EOS account
+		if (LoginSessions.IsValidIndex(LocalUserNum))
+		{
+			const FLoginSession& LoginSession = LoginSessions[LocalUserNum];
 
-					const FOnlineAccountCredentials& Creds = *GetLocalUserChecked(LocalUserNum).LastLoginCredentials_Legacy;
-					IOnlineSubsystem* PlatformOSS = GetPlatformOSS();
-					FEOSSettings Settings = UEOSSettings::GetSettings();
-					EOS_EExternalCredentialType CredType = Settings.bUseNewLoginFlow
-						? ToEOS_EExternalCredentialType(PlatformOSS ? PlatformOSS->GetSubsystemName() : EOSSubsystem->GetSubsystemName())
-						: ToEOS_EExternalCredentialTypeLegacy(PlatformOSS ? PlatformOSS->GetSubsystemName() : EOSSubsystem->GetSubsystemName(), Creds);
+			// If the user provided a Connect Auth Token, we'll try renewing the auth with it before we fire the corresponding delegate
+			if (LoginSession.UserProvidedConnectAuthToken.IsSet())
+			{
+				IOnlineSubsystem* PlatformOSS = GetPlatformOSS();
+				EOS_EExternalCredentialType CredType = ToEOS_EExternalCredentialType(PlatformOSS ? PlatformOSS->GetSubsystemName() : EOSSubsystem->GetSubsystemName());
 
-					FConnectCredentials Credentials(CredType, AuthToken);
+				FConnectCredentials Credentials(CredType, *LoginSession.UserProvidedConnectAuthToken);
 
-					EOS_Connect_LoginOptions Options = { };
-					Options.ApiVersion = 2;
-					UE_EOS_CHECK_API_MISMATCH(EOS_CONNECT_LOGIN_API_LATEST, 2);
-					Options.Credentials = &Credentials;
+				EOS_Connect_LoginOptions Options = { };
+				Options.ApiVersion = 2;
+				UE_EOS_CHECK_API_MISMATCH(EOS_CONNECT_LOGIN_API_LATEST, 2);
+				Options.Credentials = &Credentials;
 
-					FConnectLoginCallback* CallbackObj = new FConnectLoginCallback(AsWeak());
-					CallbackObj->CallbackLambda = [this, LocalUserNum](const EOS_Connect_LoginCallbackInfo* Data)
+				FConnectLoginCallback* CallbackObj = new FConnectLoginCallback(AsWeak());
+				CallbackObj->CallbackLambda = [this, LocalUserNum](const EOS_Connect_LoginCallbackInfo* Data)
 					{
 						if (Data->ResultCode != EOS_EResult::EOS_Success)
 						{
-							UE_LOG_ONLINE(Error, TEXT("Failed to refresh ConnectLogin(%d) failed with EOS result code (%s)"), LocalUserNum, *LexToString(Data->ResultCode));
-							Logout(LocalUserNum);
+							UE_LOG_ONLINE(Error, TEXT("Failed to refresh ConnectLogin for user %d. Failed with EOS result code (%s)"), LocalUserNum, *LexToString(Data->ResultCode));
+
+							TriggerOnAuthAboutToExpireDelegates(LocalUserNum);
 						}
 					};
-					EOS_Connect_Login(EOSSubsystem->ConnectHandle, &Options, CallbackObj, CallbackObj->GetCallbackPtr());
-				}
-			}));
+
+				EOS_Connect_Login(EOSSubsystem->ConnectHandle, &Options, CallbackObj, CallbackObj->GetCallbackPtr());
+			}
+			else // If not, we'll try to grab the token automatically from the platform
+			{
+				GetPlatformAuthToken(LocalUserNum,
+					FOnGetLinkedAccountAuthTokenCompleteDelegate::CreateLambda([this, WeakThis = AsWeak()](int32 LocalUserNum, bool bWasSuccessful, const FExternalAuthToken& AuthToken)
+					{
+						if (FUserManagerEOSPtr StrongThis = WeakThis.Pin())
+						{
+							if (!bWasSuccessful || !AuthToken.IsValid())
+							{
+								UE_LOG_ONLINE(Error, TEXT("AutoRefreshConnectLogin(%d) failed due to the platform OSS giving an empty auth token"), LocalUserNum);
+								Logout(LocalUserNum);
+								return;
+							}
+
+							// Now login into our EOS account
+
+							const FOnlineAccountCredentials& Creds = *GetLocalUserChecked(LocalUserNum).LastLoginCredentials_Legacy;
+							IOnlineSubsystem* PlatformOSS = GetPlatformOSS();
+							FEOSSettings Settings = UEOSSettings::GetSettings();
+							EOS_EExternalCredentialType CredType = Settings.bUseNewLoginFlow
+								? ToEOS_EExternalCredentialType(PlatformOSS ? PlatformOSS->GetSubsystemName() : EOSSubsystem->GetSubsystemName())
+								: ToEOS_EExternalCredentialTypeLegacy(PlatformOSS ? PlatformOSS->GetSubsystemName() : EOSSubsystem->GetSubsystemName(), Creds);
+
+							FConnectCredentials Credentials(CredType, ConvertExternalAuthTokenToFString(AuthToken));
+
+							EOS_Connect_LoginOptions Options = { };
+							Options.ApiVersion = 2;
+							UE_EOS_CHECK_API_MISMATCH(EOS_CONNECT_LOGIN_API_LATEST, 2);
+							Options.Credentials = &Credentials;
+
+							FConnectLoginCallback* CallbackObj = new FConnectLoginCallback(AsWeak());
+							CallbackObj->CallbackLambda = [this, LocalUserNum](const EOS_Connect_LoginCallbackInfo* Data)
+							{
+								if (Data->ResultCode != EOS_EResult::EOS_Success)
+								{
+									UE_LOG_ONLINE(Error, TEXT("Failed to refresh ConnectLogin(%d) failed with EOS result code (%s)"), LocalUserNum, *LexToString(Data->ResultCode));
+									Logout(LocalUserNum);
+								}
+							};
+							EOS_Connect_Login(EOSSubsystem->ConnectHandle, &Options, CallbackObj, CallbackObj->GetCallbackPtr());
+						}
+					}));
+			}
+		}
 	}
 }
 
@@ -1621,6 +1683,7 @@ void FUserManagerEOS::CreateConnectedLogin(int32 LocalUserNum, EOS_EpicAccountId
 			{
 				// @todo joeg - logout?
 				FString ErrorString = FString::Printf(TEXT("Login(%d) failed with EOS result code (%s)"), LocalUserNum, *LexToString(Data->ResultCode));
+				LoginSessions.RemoveAt(LocalUserNum);
 				TriggerOnLoginCompleteDelegates(LocalUserNum, false, *FUniqueNetIdEOS::EmptyId(), ErrorString);
 			}
 		};
@@ -1693,7 +1756,7 @@ void FUserManagerEOS::FullLoginCallback(int32 LocalUserNum, EOS_EpicAccountId Ac
 	}
 
 	// We add the local user
-	const FLocalUserEOS& LocalUser = AddLocalUser(LocalUserNum, AccountId, UserId);
+	FLocalUserEOS& LocalUser = AddLocalUser(LocalUserNum, AccountId, UserId);
 
 	TriggerOnLoginCompleteDelegates(LocalUserNum, true, *LocalUser.UniqueNetId, FString());
 	TriggerOnLoginStatusChangedDelegates(LocalUserNum, ELoginStatus::NotLoggedIn, ELoginStatus::LoggedIn, *LocalUser.UniqueNetId);
@@ -1722,6 +1785,8 @@ bool FUserManagerEOS::Logout(int32 LocalUserNum)
 			{
 				RemoveLocalUser(LocalUserNum);
 
+				LoginSessions.RemoveAt(LocalUserNum);
+
 				TriggerOnLogoutCompleteDelegates(LocalUserNum, true);
 			}
 			else
@@ -1820,6 +1885,11 @@ bool FUserManagerEOS::AutoLogin(int32 LocalUserNum)
 		return true;
 	}
 
+	if (!LoginSessions.IsValidIndex(LocalUserNum))
+	{
+		LoginSessions.Insert(LocalUserNum, FLoginSession());
+	}
+
 	// Check if we are configured to prefer persistent auth when attempting EAS login
 	if (Settings.bUseEAS && Settings.bPreferPersistentAuth)
 	{
@@ -1876,7 +1946,7 @@ FLocalUserEOS& FUserManagerEOS::AddLocalUser(int32 LocalUserNum, EOS_EpicAccount
 		NotificationPairPtr->Callback = CallbackObj;
 		CallbackObj->CallbackLambda = [LocalUserNum, this](const EOS_Connect_AuthExpirationCallbackInfo* Data)
 			{
-				RefreshConnectLogin(LocalUserNum);
+				AutoRefreshConnectLogin(LocalUserNum);
 			};
 
 		EOS_Connect_AddNotifyAuthExpirationOptions Options = { };
diff --git a/Engine/Plugins/Online/OnlineSubsystemEOS/Source/OnlineSubsystemEOS/Private/UserManagerEOS.h b/Engine/Plugins/Online/OnlineSubsystemEOS/Source/OnlineSubsystemEOS/Private/UserManagerEOS.h
--- a/Engine/Plugins/Online/OnlineSubsystemEOS/Source/OnlineSubsystemEOS/Private/UserManagerEOS.h
+++ b/Engine/Plugins/Online/OnlineSubsystemEOS/Source/OnlineSubsystemEOS/Private/UserManagerEOS.h
@@ -238,6 +238,13 @@ struct FLocalUserEOS
 	FNotificationIdCallbackPairPtr ConnectLoginNotification;
 };
 
+/** Created when a new local user starts a login process, keeps the state for the login session, and gets destroyed when the user logs out */
+struct FLoginSession
+{
+	/** If set, Connect Login will be re-attempted with it. If this attempt fails, a new token will need to be provided */
+	TOptional<FString> UserProvidedConnectAuthToken;
+};
+
 /**
  * EOS service implementation of the online identity, friends, and user info interfaces
  *
@@ -287,6 +294,7 @@ public:
 	virtual FPlatformUserId GetPlatformUserIdFromUniqueNetId(const FUniqueNetId& UniqueNetId) const override;
 	virtual void GetLinkedAccountAuthToken(int32 LocalUserNum, const FString& TokenType, const FOnGetLinkedAccountAuthTokenCompleteDelegate& Delegate) const override;
 // ~IOnlineIdentity Interface
+
 	ELoginStatus::Type GetLoginStatus(const FUniqueNetIdEOS& UserId) const;
 
 // IOnlineExternalUI Interface
@@ -383,7 +391,7 @@ public:
 	void LoginViaExternalAuth(int32 LocalUserNum);
 	void CreateConnectedLogin(int32 LocalUserNum, EOS_EpicAccountId AccountId, EOS_ContinuanceToken Token);
 	void LinkEAS(int32 LocalUserNum, EOS_ContinuanceToken Token);
-	void RefreshConnectLogin(int32 LocalUserNum);
+	void AutoRefreshConnectLogin(int32 LocalUserNum);
 	bool ConnectLoginNoEAS(int32 LocalUserNum);
 
 	void FullLoginCallback(int32 LocalUserNum, EOS_EpicAccountId AccountId, EOS_ProductUserId UserId);
@@ -414,6 +422,7 @@ private:
 	void CallEOSAuthLogin(int32 LocalUserNum, const FOnlineAccountCredentials& Credentials, bool bIsAutoLogin);
 	void CopyAndSaveEpicAuthToken(int32 LocalUserNum, const EOS_EpicAccountId& EpicAccountId);
 	void OnEOSAuthLoginComplete(int32 LocalUserNum, const EOS_ELoginCredentialType LoginCredentialType, const bool bIsAutoLogin, const EOS_Auth_LoginCallbackInfo* Data);
+	void CallEOSConnectLogin(int32 LocalUserNum, const FOnlineAccountCredentials& Credentials);
 
 	void RemoveLocalUser(int32 LocalUserNum);
 	FLocalUserEOS& AddLocalUser(int32 LocalUserNum, EOS_EpicAccountId EpicAccountId, EOS_ProductUserId UserId);
@@ -459,6 +468,9 @@ private:
 	/** Information related to ongoing operations and state for local users */
 	TLocalUserArray<FLocalUserEOS> LocalUsers;
 
+	/** Relevant information about the login session for local users */
+	TLocalUserArray<FLoginSession> LoginSessions;
+
 	// Online user info
 
 	/** Ids mapped to attribute access information for any user */
