OnlineSubsystemEOS is moved out of BETA.

diff --git a/Engine/Plugins/Online/OnlineSubsystemEOS/OnlineSubsystemEOS.uplugin b/Engine/Plugins/Online/OnlineSubsystemEOS/OnlineSubsystemEOS.uplugin
--- a/Engine/Plugins/Online/OnlineSubsystemEOS/OnlineSubsystemEOS.uplugin
+++ b/Engine/Plugins/Online/OnlineSubsystemEOS/OnlineSubsystemEOS.uplugin
@@ -9,7 +9,7 @@
 	"CreatedBy": "Epic Games, Inc.",
 	"CreatedByURL": "https://dev.epicgames.com/services",
 	"EnabledByDefault": false,
-	"IsBetaVersion" : true,
+	"IsBetaVersion" : false,
 
 	"Modules": [
 		{
