Added EOS Player Sanctions to OnlineSubsystemEOS
[RN] Added the EOS PlayerSanctions Interface to the OnlineSubsystemEOS

diff --git a/Engine/Plugins/Online/OnlineSubsystemEOS/Source/OnlineSubsystemEOS/Private/OnlinePlayerSanctionEOS.cpp b/Engine/Plugins/Online/OnlineSubsystemEOS/Source/OnlineSubsystemEOS/Private/OnlinePlayerSanctionEOS.cpp
new file mode 100644
--- /dev/null
+++ b/Engine/Plugins/Online/OnlineSubsystemEOS/Source/OnlineSubsystemEOS/Private/OnlinePlayerSanctionEOS.cpp
@@ -0,0 +1,149 @@
+#include "OnlinePlayerSanctionEOS.h"
+#include "OnlineSubsystem.h"
+#include "OnlineSubsystemEOS.h"
+#include "OnlineSubsystemEOSPrivate.h"
+
+#if WITH_EOS_SDK
+#include "EOSShared.h"
+#include "eos_sanctions.h"
+#include "eos_sanctions_types.h"
+#include "eos_sdk.h"
+
+DEFINE_LOG_CATEGORY(LogOnlinePlayerSanctionEOS);
+
+FOnlinePlayerSanctionEOS::FOnlinePlayerSanctionEOS(FOnlineSubsystemEOS* InSubsystem)
+	: EOSSubsystem(InSubsystem)
+{
+}
+
+// Function to convert UE enum to EOS SDK enum
+EOS_ESanctionAppealReason ToEOSSanctionAppealReason(IOnlinePlayerSanctionEOS::EPlayerSanctionAppealReason Reason)
+{
+	switch (Reason)
+	{
+	case IOnlinePlayerSanctionEOS::EPlayerSanctionAppealReason::IncorrectSanction:
+		return EOS_ESanctionAppealReason::EOS_SAR_IncorrectSanction;
+	case IOnlinePlayerSanctionEOS::EPlayerSanctionAppealReason::CompromisedAccount:
+		return EOS_ESanctionAppealReason::EOS_SAR_CompromisedAccount;
+	case IOnlinePlayerSanctionEOS::EPlayerSanctionAppealReason::UnfairPunishment:
+		return EOS_ESanctionAppealReason::EOS_SAR_UnfairPunishment;
+	case IOnlinePlayerSanctionEOS::EPlayerSanctionAppealReason::AppealForForgivenesss:
+		return EOS_ESanctionAppealReason::EOS_SAR_AppealForForgiveness; 
+	}
+	return EOS_ESanctionAppealReason::EOS_SAR_Invalid; 
+}
+
+typedef TEOSCallback<EOS_Sanctions_CreatePlayerSanctionAppealCallback, EOS_Sanctions_CreatePlayerSanctionAppealCallbackInfo, FOnlinePlayerSanctionEOS> FCreatePlayerSanctionAppealCallback;
+
+void FOnlinePlayerSanctionEOS::CreatePlayerSanctionAppeal(const FUniqueNetId& LocalUserId, FPlayerSanctionAppealSettings&& SanctionAppealSettings, FOnCreatePlayerSanctionAppealComplete&& Delegate)
+{
+
+	const FUniqueNetIdEOS& LocalEOSId = FUniqueNetIdEOS::Cast(LocalUserId);
+	const EOS_ProductUserId LocalProductUserId = LocalEOSId.GetProductUserId();
+
+	EOS_Sanctions_CreatePlayerSanctionAppealOptions Options = {};
+	Options.ApiVersion = 1;
+	UE_EOS_CHECK_API_MISMATCH(EOS_SANCTIONS_CREATEPLAYERSANCTIONAPPEAL_API_LATEST, 1);
+	Options.LocalUserId = LocalProductUserId; 
+	Options.Reason = ToEOSSanctionAppealReason(SanctionAppealSettings.Reason);
+	const auto ReferenceIdUTF8 = StringCast<UTF8CHAR>(*SanctionAppealSettings.ReferenceId);
+	Options.ReferenceId = (const char*)ReferenceIdUTF8.Get();
+
+	FCreatePlayerSanctionAppealCallback* CallbackObj = new FCreatePlayerSanctionAppealCallback(AsWeak());
+	CallbackObj->CallbackLambda = [Delegate = MoveTemp(Delegate)](const EOS_Sanctions_CreatePlayerSanctionAppealCallbackInfo* Data)
+		{
+			const bool bWasSuccessful = Data->ResultCode == EOS_EResult::EOS_Success;
+
+			if (!bWasSuccessful)
+			{
+				UE_LOG_ONLINE_PLAYERSANCTIONEOS(Warning, TEXT("Failed to send sanction appeal for local user."));
+			}
+			Delegate.ExecuteIfBound(bWasSuccessful); 
+		};
+
+	EOS_Sanctions_CreatePlayerSanctionAppeal(EOSSubsystem->PlayerSanctionHandle, &Options, CallbackObj, CallbackObj->GetCallbackPtr()); 
+}
+
+typedef TEOSCallback<EOS_Sanctions_OnQueryActivePlayerSanctionsCallback, EOS_Sanctions_QueryActivePlayerSanctionsCallbackInfo, FOnlinePlayerSanctionEOS> FQueryActivePlayerSanctionsCallback;
+
+void FOnlinePlayerSanctionEOS::QueryActivePlayerSanctions(const FUniqueNetId& LocalUserId, const FUniqueNetId& TargetUserId, FOnQueryActivePlayerSanctionsComplete&& Delegate)
+{
+	const FUniqueNetIdEOS& LocalEOSId = FUniqueNetIdEOS::Cast(LocalUserId);
+	const EOS_ProductUserId LocalProductUserId = LocalEOSId.GetProductUserId();
+
+	const FUniqueNetIdEOS& TargetEOSId = FUniqueNetIdEOS::Cast(TargetUserId);
+	const EOS_ProductUserId TargetProductUserId = TargetEOSId.GetProductUserId();
+
+	EOS_Sanctions_QueryActivePlayerSanctionsOptions QueryOptions = {}; 
+	QueryOptions.ApiVersion = 2;
+	UE_EOS_CHECK_API_MISMATCH(EOS_SANCTIONS_QUERYACTIVEPLAYERSANCTIONS_API_LATEST, 2);
+	QueryOptions.LocalUserId = LocalProductUserId; 
+	QueryOptions.TargetUserId = TargetProductUserId; 
+
+	FQueryActivePlayerSanctionsCallback* CallbackObj = new FQueryActivePlayerSanctionsCallback(AsWeak());
+	CallbackObj->CallbackLambda = [this, LambdaPlayerId = TargetUserId.AsShared(), Delegate = MoveTemp(Delegate)](const  EOS_Sanctions_QueryActivePlayerSanctionsCallbackInfo* Data)
+		{
+			const bool bWasSuccessful = Data->ResultCode == EOS_EResult::EOS_Success;
+
+			if (bWasSuccessful)
+			{
+				TArray<FOnlinePlayerSanction>& PlayerSanctions = CachedPlayerSanctionsMap.Emplace(LambdaPlayerId);
+
+				const FUniqueNetIdEOS& TargetEOSId = FUniqueNetIdEOS::Cast(*LambdaPlayerId);
+				const EOS_ProductUserId TargetProductUserId = TargetEOSId.GetProductUserId();
+
+				EOS_Sanctions_GetPlayerSanctionCountOptions CountOptions = {};
+				CountOptions.ApiVersion = 1; 
+				UE_EOS_CHECK_API_MISMATCH(EOS_SANCTIONS_GETPLAYERSANCTIONCOUNT_API_LATEST, 1);
+				CountOptions.TargetUserId = TargetProductUserId;
+
+				uint32 SanctionCount = EOS_Sanctions_GetPlayerSanctionCount(EOSSubsystem->PlayerSanctionHandle, &CountOptions); 
+
+				EOS_Sanctions_CopyPlayerSanctionByIndexOptions CopyOptions = {};
+				CopyOptions.ApiVersion = 1; 
+				UE_EOS_CHECK_API_MISMATCH(EOS_SANCTIONS_COPYPLAYERSANCTIONBYINDEX_API_LATEST, 1);
+				CopyOptions.TargetUserId = TargetProductUserId;
+
+				for (uint32 SanctionIndex = 0; SanctionIndex < SanctionCount; SanctionIndex++)
+				{
+					CopyOptions.SanctionIndex = SanctionIndex;
+
+					EOS_Sanctions_PlayerSanction* PlayerSanctionEOS = nullptr; 
+
+					EOS_EResult Result = EOS_Sanctions_CopyPlayerSanctionByIndex(EOSSubsystem->PlayerSanctionHandle, &CopyOptions, &PlayerSanctionEOS);
+					if (Result == EOS_EResult::EOS_Success)
+					{
+						FOnlinePlayerSanction& PlayerSanction = PlayerSanctions.Emplace_GetRef();
+						PlayerSanction.TimePlaced = PlayerSanctionEOS->TimePlaced;
+						PlayerSanction.TimeExpires = PlayerSanctionEOS->TimeExpires;
+						PlayerSanction.ReferenceId = PlayerSanctionEOS->ReferenceId;
+						PlayerSanction.Action = PlayerSanctionEOS->Action;
+
+						EOS_Sanctions_PlayerSanction_Release(PlayerSanctionEOS); 
+
+					}
+					else
+					{
+						UE_LOG_ONLINE_PLAYERSANCTIONEOS(Warning, TEXT("EOS_Sanctions_CopyPlayerSanctionByIndex() failed for player (%s). The error code is: (%s)"), *LambdaPlayerId->ToString(),*LexToString(Result));
+					}
+				}
+			}
+			else
+			{
+				UE_LOG_ONLINE_PLAYERSANCTIONEOS(Warning, TEXT("Failed to query active player sanctions for player (%s)."), *LambdaPlayerId->ToString());
+			}
+			Delegate.ExecuteIfBound(bWasSuccessful);
+		};
+
+	EOS_Sanctions_QueryActivePlayerSanctions(EOSSubsystem->PlayerSanctionHandle, &QueryOptions, CallbackObj, CallbackObj->GetCallbackPtr()); 
+}
+EOnlineCachedResult::Type FOnlinePlayerSanctionEOS::GetCachedActivePlayerSanctions(const FUniqueNetId& TargetUserId, TArray<FOnlinePlayerSanction>& OutPlayerSanctions)
+{
+	if (const TArray<FOnlinePlayerSanction>* CachedPlayerSanctions = CachedPlayerSanctionsMap.Find(TargetUserId.AsShared()))
+	{
+		OutPlayerSanctions = *CachedPlayerSanctions;
+		return EOnlineCachedResult::Success;
+	}
+	return EOnlineCachedResult::NotFound;
+}
+#endif
\ No newline at end of file
diff --git a/Engine/Plugins/Online/OnlineSubsystemEOS/Source/OnlineSubsystemEOS/Private/OnlinePlayerSanctionEOS.h b/Engine/Plugins/Online/OnlineSubsystemEOS/Source/OnlineSubsystemEOS/Private/OnlinePlayerSanctionEOS.h
new file mode 100644
--- /dev/null
+++ b/Engine/Plugins/Online/OnlineSubsystemEOS/Source/OnlineSubsystemEOS/Private/OnlinePlayerSanctionEOS.h
@@ -0,0 +1,48 @@
+// Copyright Epic Games, Inc. All Rights Reserved.
+
+#pragma once
+
+#include "Misc/ScopeLock.h"
+#include "Interfaces/OnlinePlayerSanctionEOSInterface.h"
+#include "OnlineSubsystemEOSTypes.h"
+
+DECLARE_LOG_CATEGORY_EXTERN(LogOnlinePlayerSanctionEOS, Log, All);
+
+#define UE_LOG_ONLINE_PLAYERSANCTIONEOS(Verbosity, Format, ...) \
+{ \
+	UE_LOG(LogOnlinePlayerSanctionEOS, Verbosity, TEXT("%s%s"), ONLINE_LOG_PREFIX, *FString::Printf(Format, ##__VA_ARGS__)); \
+}
+
+class FOnlineSubsystemEOS;
+
+#if WITH_EOS_SDK
+
+/**
+ * Interface for interacting with EOS sanctions
+ */
+class FOnlinePlayerSanctionEOS
+	: public IOnlinePlayerSanctionEOS
+	, public TSharedFromThis<FOnlinePlayerSanctionEOS, ESPMode::ThreadSafe>
+{
+public:
+	FOnlinePlayerSanctionEOS() = delete;
+	virtual ~FOnlinePlayerSanctionEOS() = default;
+
+	FOnlinePlayerSanctionEOS(FOnlineSubsystemEOS* InSubsystem);
+
+	virtual void CreatePlayerSanctionAppeal(const FUniqueNetId& LocalUserId, FPlayerSanctionAppealSettings&& SanctionAppealSettings, FOnCreatePlayerSanctionAppealComplete&& Delegate) override;
+
+	virtual void QueryActivePlayerSanctions(const FUniqueNetId& LocalUserId, const FUniqueNetId& TargetUserId, FOnQueryActivePlayerSanctionsComplete&& Delegate) override;
+
+	virtual  EOnlineCachedResult::Type GetCachedActivePlayerSanctions(const FUniqueNetId& TargetUserId, TArray<FOnlinePlayerSanction>& OutPlayerSanctions) override;
+
+
+private:
+	/** Reference to the main EOS subsystem */
+	FOnlineSubsystemEOS* EOSSubsystem;
+
+	/** Holds the cached info from the last time this was called */
+	TUniqueNetIdMap<TArray<FOnlinePlayerSanction>> CachedPlayerSanctionsMap;
+};
+
+#endif
diff --git a/Engine/Plugins/Online/OnlineSubsystemEOS/Source/OnlineSubsystemEOS/Private/OnlineSubsystemEOS.cpp b/Engine/Plugins/Online/OnlineSubsystemEOS/Source/OnlineSubsystemEOS/Private/OnlineSubsystemEOS.cpp
--- a/Engine/Plugins/Online/OnlineSubsystemEOS/Source/OnlineSubsystemEOS/Private/OnlineSubsystemEOS.cpp
+++ b/Engine/Plugins/Online/OnlineSubsystemEOS/Source/OnlineSubsystemEOS/Private/OnlineSubsystemEOS.cpp
@@ -11,6 +11,7 @@
 #include "OnlineTitleFileEOS.h"
 #include "OnlineUserCloudEOS.h"
 #include "OnlineStoreEOS.h"
+#include "OnlinePlayerSanctionEOS.h"
 #include "EOSSettings.h"
 #include "EOSShared.h"
 #include "EOSVoiceChat.h"
@@ -356,6 +357,12 @@ bool FOnlineSubsystemEOS::Init()
 		UE_LOG_ONLINE(Error, TEXT("FOnlineSubsystemEOS: failed to init EOS platform, couldn't get ecom handle"));
 		return false;
 	}
+	PlayerSanctionHandle = EOS_Platform_GetSanctionsInterface(*EOSPlatformHandle);
+	if (PlayerSanctionHandle == nullptr)
+	{
+		UE_LOG_ONLINE(Error, TEXT("FOnlineSubsystemEOS: failed to init EOS platform, couldn't get player sanction handle"));
+		return false;
+	}
 
 	// We set the product id
 	FEOSArtifactSettings ArtifactSettings;
@@ -408,14 +415,15 @@ bool FOnlineSubsystemEOS::Init()
 		UE_LOG_ONLINE(Warning, TEXT("[FOnlineSubsystemEOS::Init] Unable to initialize Socket Subsystem. Error=[%s]"), *ErrorMessage);
 	}
 
-	UserManager = MakeShareable(new FUserManagerEOS(this));
+	UserManager = MakeShared<FUserManagerEOS>(this);
 	UserManager->Init();
-	SessionInterfacePtr = MakeShareable(new FOnlineSessionEOS(this));
+	SessionInterfacePtr = MakeShared<FOnlineSessionEOS>(this);
 	SessionInterfacePtr->Init();
-	StatsInterfacePtr = MakeShareable(new FOnlineStatsEOS(this));
-	LeaderboardsInterfacePtr = MakeShareable(new FOnlineLeaderboardsEOS(this));
-	AchievementsInterfacePtr = MakeShareable(new FOnlineAchievementsEOS(this));
-	StoreInterfacePtr = MakeShareable(new FOnlineStoreEOS(this));
+	StatsInterfacePtr = MakeShared<FOnlineStatsEOS>(this);
+	LeaderboardsInterfacePtr = MakeShared<FOnlineLeaderboardsEOS>(this);
+	AchievementsInterfacePtr = MakeShared<FOnlineAchievementsEOS>(this);
+	StoreInterfacePtr = MakeShared<FOnlineStoreEOS>(this);
+	PlayerSanctionEOSPtr = MakeShared<FOnlinePlayerSanctionEOS>(this);
 
 	// We initialized ok so we can tick
 	StartTicker();
@@ -675,6 +683,11 @@ IOnlineStatsPtr FOnlineSubsystemEOS::GetStatsInterface() const
 	return StatsInterfacePtr;
 }
 
+IOnlinePlayerSanctionEOSPtr FOnlineSubsystemEOS::GetPlayerSanctionEOSInterface() const
+{
+	return PlayerSanctionEOSPtr;
+}
+
 IVoiceChatUser* FOnlineSubsystemEOS::GetVoiceChatUserInterface(const FUniqueNetId& LocalUserId)
 {
 	IVoiceChatUser* Result = nullptr;
diff --git a/Engine/Plugins/Online/OnlineSubsystemEOS/Source/OnlineSubsystemEOS/Private/OnlineSubsystemEOS.h b/Engine/Plugins/Online/OnlineSubsystemEOS/Source/OnlineSubsystemEOS/Private/OnlineSubsystemEOS.h
--- a/Engine/Plugins/Online/OnlineSubsystemEOS/Source/OnlineSubsystemEOS/Private/OnlineSubsystemEOS.h
+++ b/Engine/Plugins/Online/OnlineSubsystemEOS/Source/OnlineSubsystemEOS/Private/OnlineSubsystemEOS.h
@@ -47,6 +47,9 @@ typedef TSharedPtr<class FOnlineTitleFileEOS, ESPMode::ThreadSafe> FOnlineTitleF
 class FOnlineUserCloudEOS;
 typedef TSharedPtr<class FOnlineUserCloudEOS, ESPMode::ThreadSafe> FOnlineUserCloudEOSPtr;
 
+class FOnlinePlayerSanctionEOS;
+typedef TSharedPtr<class FOnlinePlayerSanctionEOS, ESPMode::ThreadSafe> FOnlinePlayerSanctionEOSPtr;
+
 typedef TSharedPtr<FPlatformEOSHelpers, ESPMode::ThreadSafe> FPlatformEOSHelpersPtr;
 
 /**
@@ -86,6 +89,7 @@ public:
 	virtual IOnlinePresencePtr GetPresenceInterface() const override;
 	virtual FText GetOnlineServiceName() const override;
 	virtual IOnlineStatsPtr GetStatsInterface() const override;
+	virtual IOnlinePlayerSanctionEOSPtr GetPlayerSanctionEOSInterface() const override;
 	virtual bool Exec(class UWorld* InWorld, const TCHAR* Cmd, FOutputDevice& Ar) override;
 	virtual void ReloadConfigs(const TSet<FString>& ConfigSections) override;
 
@@ -131,6 +135,7 @@ public:
 	EOS_HEcom EcomHandle;
 	EOS_HTitleStorage TitleStorageHandle;
 	EOS_HPlayerDataStorage PlayerDataStorageHandle;
+	EOS_HSanctions PlayerSanctionHandle;
 
 	/** Manager that handles all user interfaces */
 	FUserManagerEOSPtr UserManager;
@@ -147,6 +152,8 @@ public:
 	FOnlineTitleFileEOSPtr TitleFileInterfacePtr;
 	/** User Cloud interface pointer */
 	FOnlineUserCloudEOSPtr UserCloudInterfacePtr;
+	/** Player sanction interface pointer */
+	FOnlinePlayerSanctionEOSPtr PlayerSanctionEOSPtr;
 
 	TSharedPtr<FSocketSubsystemEOS, ESPMode::ThreadSafe> SocketSubsystem;
 
diff --git a/Engine/Plugins/Online/OnlineSubsystemEOS/Source/OnlineSubsystemEOS/Public/IOnlineSubsystemEOS.h b/Engine/Plugins/Online/OnlineSubsystemEOS/Source/OnlineSubsystemEOS/Public/IOnlineSubsystemEOS.h
--- a/Engine/Plugins/Online/OnlineSubsystemEOS/Source/OnlineSubsystemEOS/Public/IOnlineSubsystemEOS.h
+++ b/Engine/Plugins/Online/OnlineSubsystemEOS/Source/OnlineSubsystemEOS/Public/IOnlineSubsystemEOS.h
@@ -7,6 +7,7 @@
 class FUniqueNetId;
 class IVoiceChatUser;
 using IEOSPlatformHandlePtr = TSharedPtr<class IEOSPlatformHandle, ESPMode::ThreadSafe>;
+using IOnlinePlayerSanctionEOSPtr = TSharedPtr<class IOnlinePlayerSanctionEOS, ESPMode::ThreadSafe>;
 
 /**
  *	OnlineSubsystemEOS - Implementation of the online subsystem for EOS services
@@ -20,4 +21,5 @@ public:
 
 	virtual IVoiceChatUser* GetVoiceChatUserInterface(const FUniqueNetId& LocalUserId) = 0;
 	virtual IEOSPlatformHandlePtr GetEOSPlatformHandle() const = 0;
+	virtual IOnlinePlayerSanctionEOSPtr GetPlayerSanctionEOSInterface() const = 0;
 };
\ No newline at end of file
diff --git a/Engine/Plugins/Online/OnlineSubsystemEOS/Source/OnlineSubsystemEOS/Public/Interfaces/OnlinePlayerSanctionEOSInterface.h b/Engine/Plugins/Online/OnlineSubsystemEOS/Source/OnlineSubsystemEOS/Public/Interfaces/OnlinePlayerSanctionEOSInterface.h
new file mode 100644
--- /dev/null
+++ b/Engine/Plugins/Online/OnlineSubsystemEOS/Source/OnlineSubsystemEOS/Public/Interfaces/OnlinePlayerSanctionEOSInterface.h
@@ -0,0 +1,91 @@
+#pragma once
+
+#include "Online/CoreOnline.h"
+#include "OnlineSubsystemTypes.h"
+
+/**
+ * Delegate fired when a player sanction appeal has been created
+ *
+ *  @param bWasSuccessful true if creating the sanction appel was successful
+ */
+DECLARE_DELEGATE_OneParam(FOnCreatePlayerSanctionAppealComplete, const bool /* bWasSuccessful */);
+
+/**
+ * Delegate fired when player sanctions have been cached locally
+ *
+ *  @param bWasSuccessful true if sanctions are available in local cache
+ */
+DECLARE_DELEGATE_OneParam(FOnQueryActivePlayerSanctionsComplete, const bool /* bWasSuccessful */);
+
+/**
+ * Public Interface for interacting with EOS player sanctions
+ */
+class IOnlinePlayerSanctionEOS
+{
+public:
+	virtual ~IOnlinePlayerSanctionEOS() {};
+
+
+	/**
+	*	Enum to set the appeal reason when creating a sanction appeal
+	*/
+	enum class EPlayerSanctionAppealReason
+	{
+		IncorrectSanction,
+		CompromisedAccount,
+		UnfairPunishment,
+		AppealForForgivenesss
+	};
+
+	/**
+	*	Struct for the settings of sanction appeal
+	*/
+	struct FPlayerSanctionAppealSettings
+	{
+		/** The reason of the appeal */
+		EPlayerSanctionAppealReason Reason;
+
+		/** The sanction id for the sanction that is being appealed */
+		FString ReferenceId;
+	};
+
+	/**
+	*	Struct for player sanctions
+	*/
+	struct FOnlinePlayerSanction
+	{
+		/** The time the sanction was placed */
+		int64_t TimePlaced;
+
+		/** The time the sanction expires */
+		int64_t TimeExpires;
+
+		/** The action associated with this sanction */
+		FString Action;
+
+		/** The sanction id for the sanction that is being appealed. This needs to be set */
+		FString ReferenceId;
+	};
+
+	/**
+	* Send a player sanction appeal
+	* @param LocalUserId the player id of the local user sending the sanction appeal. 
+	* @param SanctionAppealSettings the settings needed to create the sanction appeal.
+	 */
+	virtual void CreatePlayerSanctionAppeal(const FUniqueNetId& LocalUserId, FPlayerSanctionAppealSettings&& SanctionAppealSettings, FOnCreatePlayerSanctionAppealComplete&& Delegate) = 0;
+
+	/**
+	* Query active player sanction. The sanctions will be cached locally and can be retrieved using GetCachedActivePlayerSanctions
+	* @param LocalUserId the player id of the local user querying the sanctions. 
+	* @param TargetUserId the player id of the player we are querying sanctions for. 
+	 */
+	virtual void QueryActivePlayerSanctions(const FUniqueNetId& LocalUserId, const FUniqueNetId& TargetUserId, FOnQueryActivePlayerSanctionsComplete&& Delegate) = 0;
+
+	/**
+	* Retrieve cached player sanctions
+	* @param TargetUserId the player id we want to retrieve the sanctions for.
+	* @param OutPlayerSanctions the array to store the sanctions.
+	*/
+	virtual EOnlineCachedResult::Type GetCachedActivePlayerSanctions(const FUniqueNetId& TargetUserId, TArray<FOnlinePlayerSanction>& OutPlayerSanctions) = 0;
+};
+typedef TSharedPtr<IOnlinePlayerSanctionEOS, ESPMode::ThreadSafe> IOnlinePlayerSanctionEOSPtr;
\ No newline at end of file
