Added EOS Player Reports to OnlineSubsystemEOS
[RN] Added the EOS Player Reports interface to the OnlineSubsystemEOS

diff --git a/Engine/Plugins/Online/OnlineSubsystemEOS/Source/OnlineSubsystemEOS/Private/OnlinePlayerReportEOS.cpp b/Engine/Plugins/Online/OnlineSubsystemEOS/Source/OnlineSubsystemEOS/Private/OnlinePlayerReportEOS.cpp
new file mode 100644
--- /dev/null
+++ b/Engine/Plugins/Online/OnlineSubsystemEOS/Source/OnlineSubsystemEOS/Private/OnlinePlayerReportEOS.cpp
@@ -0,0 +1,89 @@
+// Copyright Epic Games, Inc. All Rights Reserved.
+
+#include "OnlinePlayerReportEOS.h"
+#include "OnlineSubsystem.h"
+#include "OnlineSubsystemEOS.h"
+#include "OnlineSubsystemEOSPrivate.h"
+
+#if WITH_EOS_SDK
+#include "eos_reports.h"
+#include "eos_reports_types.h"
+#include "eos_sdk.h"
+
+DEFINE_LOG_CATEGORY(LogOnlinePlayerReportEOS); 
+
+FOnlinePlayerReportEOS::FOnlinePlayerReportEOS(FOnlineSubsystemEOS* InSubsystem)
+	: EOSSubsystem(InSubsystem)
+{
+}
+
+// Convert the plugin player report category enum to the EOS SDK enum
+EOS_EPlayerReportsCategory ToEOSPlayerReportsCategory(IOnlinePlayerReportEOS::EPlayerReportCategory Category)
+{
+	switch (Category)
+	{
+	case IOnlinePlayerReportEOS::EPlayerReportCategory::Cheating:
+	{
+		return EOS_EPlayerReportsCategory::EOS_PRC_Cheating;
+	}
+	case IOnlinePlayerReportEOS::EPlayerReportCategory::Exploiting:
+	{
+		return EOS_EPlayerReportsCategory::EOS_PRC_Exploiting;
+	}
+	case IOnlinePlayerReportEOS::EPlayerReportCategory::OffensiveProfile:
+	{
+		return EOS_EPlayerReportsCategory::EOS_PRC_OffensiveProfile;
+	}
+	case IOnlinePlayerReportEOS::EPlayerReportCategory::VerbalAbuse:
+	{
+		return EOS_EPlayerReportsCategory::EOS_PRC_VerbalAbuse;
+	}
+	case IOnlinePlayerReportEOS::EPlayerReportCategory::Scamming:
+	{
+		return EOS_EPlayerReportsCategory::EOS_PRC_Scamming;
+	}
+	case IOnlinePlayerReportEOS::EPlayerReportCategory::Spamming:
+	{
+		return EOS_EPlayerReportsCategory::EOS_PRC_Spamming;
+	}
+	}
+	// Default to other
+	return EOS_EPlayerReportsCategory::EOS_PRC_Other;
+}
+
+typedef TEOSCallback<EOS_Reports_OnSendPlayerBehaviorReportCompleteCallback, EOS_Reports_SendPlayerBehaviorReportCompleteCallbackInfo, FOnlinePlayerReportEOS> FSendPlayerReportCallback;
+
+void FOnlinePlayerReportEOS::SendPlayerReport(const FUniqueNetId& LocalUserId, const FUniqueNetId& TargetUserId, FSendPlayerReportSettings&& SendPlayerReportSettings, FOnSendPlayerReportComplete&& Delegate)
+{
+	const FUniqueNetIdEOS& LocalEOSId = FUniqueNetIdEOS::Cast(LocalUserId);
+	const EOS_ProductUserId LocalProductUserId = LocalEOSId.GetProductUserId();
+
+	const FUniqueNetIdEOS& TargetEOSId = FUniqueNetIdEOS::Cast(TargetUserId);
+	const EOS_ProductUserId TargetProductUserId = TargetEOSId.GetProductUserId();
+
+	EOS_Reports_SendPlayerBehaviorReportOptions Options = {};
+	Options.ApiVersion = 2;
+	UE_EOS_CHECK_API_MISMATCH(EOS_REPORTS_SENDPLAYERBEHAVIORREPORT_API_LATEST, 2);
+	Options.ReporterUserId = LocalProductUserId;
+    Options.ReportedUserId = TargetProductUserId;
+	Options.Category = ToEOSPlayerReportsCategory(SendPlayerReportSettings.Category);
+	const auto MessagedUTF8 = StringCast<UTF8CHAR>(*SendPlayerReportSettings.Message);
+	const auto ContextUTF8 = StringCast<UTF8CHAR>(*SendPlayerReportSettings.Context);
+	Options.Message = SendPlayerReportSettings.Message.IsEmpty() ? nullptr : (const char*)MessagedUTF8.Get();
+	Options.Context = SendPlayerReportSettings.Context.IsEmpty() ? nullptr : (const char*)ContextUTF8.Get();
+	
+	FSendPlayerReportCallback* CallbackObj = new FSendPlayerReportCallback(AsWeak());
+	CallbackObj->CallbackLambda = [LocalProductUserId, TargetProductUserId, Delegate = MoveTemp(Delegate)](const EOS_Reports_SendPlayerBehaviorReportCompleteCallbackInfo* Data)
+		{
+			const bool bWasSuccessful = Data->ResultCode == EOS_EResult::EOS_Success;
+
+			if (!bWasSuccessful)
+			{
+				UE_LOG_ONLINE_PLAYERREPORTEOS(Warning, TEXT("LocalUserId (%s) failed to send the player report for TargetUserId(%s). The error code is: (%s)."),*LexToString(LocalProductUserId), *LexToString(TargetProductUserId), *LexToString(Data->ResultCode));
+			}
+			Delegate.ExecuteIfBound(bWasSuccessful);
+		};
+
+	EOS_Reports_SendPlayerBehaviorReport(EOSSubsystem->PlayerReportHandle, &Options, CallbackObj, CallbackObj->GetCallbackPtr());
+}
+#endif
diff --git a/Engine/Plugins/Online/OnlineSubsystemEOS/Source/OnlineSubsystemEOS/Private/OnlinePlayerReportEOS.h b/Engine/Plugins/Online/OnlineSubsystemEOS/Source/OnlineSubsystemEOS/Private/OnlinePlayerReportEOS.h
new file mode 100644
--- /dev/null
+++ b/Engine/Plugins/Online/OnlineSubsystemEOS/Source/OnlineSubsystemEOS/Private/OnlinePlayerReportEOS.h
@@ -0,0 +1,39 @@
+// Copyright Epic Games, Inc. All Rights Reserved.
+
+#pragma once
+
+#include "Misc/ScopeLock.h"
+#include "Interfaces/OnlinePlayerReportEOSInterface.h"
+#include "OnlineSubsystemEOSTypes.h"
+
+DECLARE_LOG_CATEGORY_EXTERN(LogOnlinePlayerReportEOS, Log, All);
+
+#define UE_LOG_ONLINE_PLAYERREPORTEOS(Verbosity, Format, ...) \
+{ \
+	UE_LOG(LogOnlinePlayerReportEOS, Verbosity, TEXT("%s%s"), ONLINE_LOG_PREFIX, *FString::Printf(Format, ##__VA_ARGS__)); \
+}
+
+class FOnlineSubsystemEOS;
+
+#if WITH_EOS_SDK
+
+/**
+ * Interface for interacting with EOS player reports
+ */
+class FOnlinePlayerReportEOS
+	: public IOnlinePlayerReportEOS
+	, public TSharedFromThis<FOnlinePlayerReportEOS, ESPMode::ThreadSafe>
+{
+public:
+	FOnlinePlayerReportEOS() = delete;
+	virtual ~FOnlinePlayerReportEOS() = default;
+
+	FOnlinePlayerReportEOS(FOnlineSubsystemEOS* InSubsystem);
+
+	virtual void SendPlayerReport(const FUniqueNetId& LocalUserId, const FUniqueNetId& TargetUserId, FSendPlayerReportSettings&& SendPlayerReportSettings, FOnSendPlayerReportComplete&& Delegate) override;
+
+private:
+	/** Reference to the main EOS subsystem */
+	FOnlineSubsystemEOS* EOSSubsystem;
+};
+#endif
\ No newline at end of file
diff --git a/Engine/Plugins/Online/OnlineSubsystemEOS/Source/OnlineSubsystemEOS/Private/OnlineSubsystemEOS.cpp b/Engine/Plugins/Online/OnlineSubsystemEOS/Source/OnlineSubsystemEOS/Private/OnlineSubsystemEOS.cpp
--- a/Engine/Plugins/Online/OnlineSubsystemEOS/Source/OnlineSubsystemEOS/Private/OnlineSubsystemEOS.cpp
+++ b/Engine/Plugins/Online/OnlineSubsystemEOS/Source/OnlineSubsystemEOS/Private/OnlineSubsystemEOS.cpp
@@ -12,6 +12,7 @@
 #include "OnlineUserCloudEOS.h"
 #include "OnlineStoreEOS.h"
 #include "OnlinePlayerSanctionEOS.h"
+#include "OnlinePlayerReportEOS.h"
 #include "EOSSettings.h"
 #include "EOSShared.h"
 #include "EOSVoiceChat.h"
@@ -363,6 +364,12 @@ bool FOnlineSubsystemEOS::Init()
 		UE_LOG_ONLINE(Error, TEXT("FOnlineSubsystemEOS: failed to init EOS platform, couldn't get player sanction handle"));
 		return false;
 	}
+	PlayerReportHandle = EOS_Platform_GetReportsInterface(*EOSPlatformHandle);
+	if (PlayerReportHandle == nullptr)
+	{
+		UE_LOG_ONLINE(Error, TEXT("FOnlineSubsystemEOS: failed to init EOS platform, couldn't get player report handle"));
+		return false;
+	}
 
 	// We set the product id
 	FEOSArtifactSettings ArtifactSettings;
@@ -424,6 +431,7 @@ bool FOnlineSubsystemEOS::Init()
 	AchievementsInterfacePtr = MakeShared<FOnlineAchievementsEOS>(this);
 	StoreInterfacePtr = MakeShared<FOnlineStoreEOS>(this);
 	PlayerSanctionEOSPtr = MakeShared<FOnlinePlayerSanctionEOS>(this);
+	PlayerReportInterfacePtr = MakeShared<FOnlinePlayerReportEOS>(this);
 
 	// We initialized ok so we can tick
 	StartTicker();
@@ -688,6 +696,11 @@ IOnlinePlayerSanctionEOSPtr FOnlineSubsystemEOS::GetPlayerSanctionEOSInterface()
 	return PlayerSanctionEOSPtr;
 }
 
+IOnlinePlayerReportEOSPtr FOnlineSubsystemEOS::GetPlayerReportEOSInterface() const
+{
+	return PlayerReportInterfacePtr;
+}
+
 IVoiceChatUser* FOnlineSubsystemEOS::GetVoiceChatUserInterface(const FUniqueNetId& LocalUserId)
 {
 	IVoiceChatUser* Result = nullptr;
diff --git a/Engine/Plugins/Online/OnlineSubsystemEOS/Source/OnlineSubsystemEOS/Private/OnlineSubsystemEOS.h b/Engine/Plugins/Online/OnlineSubsystemEOS/Source/OnlineSubsystemEOS/Private/OnlineSubsystemEOS.h
--- a/Engine/Plugins/Online/OnlineSubsystemEOS/Source/OnlineSubsystemEOS/Private/OnlineSubsystemEOS.h
+++ b/Engine/Plugins/Online/OnlineSubsystemEOS/Source/OnlineSubsystemEOS/Private/OnlineSubsystemEOS.h
@@ -50,6 +50,9 @@ typedef TSharedPtr<class FOnlineUserCloudEOS, ESPMode::ThreadSafe> FOnlineUserCl
 class FOnlinePlayerSanctionEOS;
 typedef TSharedPtr<class FOnlinePlayerSanctionEOS, ESPMode::ThreadSafe> FOnlinePlayerSanctionEOSPtr;
 
+class FOnlinePlayerReportEOS;
+typedef TSharedPtr<class FOnlinePlayerReportEOS, ESPMode::ThreadSafe> FOnlinePlayerReportEOSPtr;
+
 typedef TSharedPtr<FPlatformEOSHelpers, ESPMode::ThreadSafe> FPlatformEOSHelpersPtr;
 
 /**
@@ -90,6 +93,7 @@ public:
 	virtual FText GetOnlineServiceName() const override;
 	virtual IOnlineStatsPtr GetStatsInterface() const override;
 	virtual IOnlinePlayerSanctionEOSPtr GetPlayerSanctionEOSInterface() const override;
+	virtual IOnlinePlayerReportEOSPtr GetPlayerReportEOSInterface() const override;
 	virtual bool Exec(class UWorld* InWorld, const TCHAR* Cmd, FOutputDevice& Ar) override;
 	virtual void ReloadConfigs(const TSet<FString>& ConfigSections) override;
 
@@ -136,6 +140,7 @@ public:
 	EOS_HTitleStorage TitleStorageHandle;
 	EOS_HPlayerDataStorage PlayerDataStorageHandle;
 	EOS_HSanctions PlayerSanctionHandle;
+	EOS_HReports PlayerReportHandle;
 
 	/** Manager that handles all user interfaces */
 	FUserManagerEOSPtr UserManager;
@@ -154,6 +159,8 @@ public:
 	FOnlineUserCloudEOSPtr UserCloudInterfacePtr;
 	/** Player sanction interface pointer */
 	FOnlinePlayerSanctionEOSPtr PlayerSanctionEOSPtr;
+	/** Player Report interface pointer */
+	FOnlinePlayerReportEOSPtr PlayerReportInterfacePtr;
 
 	TSharedPtr<FSocketSubsystemEOS, ESPMode::ThreadSafe> SocketSubsystem;
 
diff --git a/Engine/Plugins/Online/OnlineSubsystemEOS/Source/OnlineSubsystemEOS/Public/IOnlineSubsystemEOS.h b/Engine/Plugins/Online/OnlineSubsystemEOS/Source/OnlineSubsystemEOS/Public/IOnlineSubsystemEOS.h
--- a/Engine/Plugins/Online/OnlineSubsystemEOS/Source/OnlineSubsystemEOS/Public/IOnlineSubsystemEOS.h
+++ b/Engine/Plugins/Online/OnlineSubsystemEOS/Source/OnlineSubsystemEOS/Public/IOnlineSubsystemEOS.h
@@ -8,6 +8,7 @@ class FUniqueNetId;
 class IVoiceChatUser;
 using IEOSPlatformHandlePtr = TSharedPtr<class IEOSPlatformHandle, ESPMode::ThreadSafe>;
 using IOnlinePlayerSanctionEOSPtr = TSharedPtr<class IOnlinePlayerSanctionEOS, ESPMode::ThreadSafe>;
+using IOnlinePlayerReportEOSPtr = TSharedPtr<class IOnlinePlayerReportEOS, ESPMode::ThreadSafe>;
 
 /**
  *	OnlineSubsystemEOS - Implementation of the online subsystem for EOS services
@@ -22,4 +23,5 @@ public:
 	virtual IVoiceChatUser* GetVoiceChatUserInterface(const FUniqueNetId& LocalUserId) = 0;
 	virtual IEOSPlatformHandlePtr GetEOSPlatformHandle() const = 0;
 	virtual IOnlinePlayerSanctionEOSPtr GetPlayerSanctionEOSInterface() const = 0;
+	virtual IOnlinePlayerReportEOSPtr GetPlayerReportEOSInterface() const = 0;
 };
\ No newline at end of file
diff --git a/Engine/Plugins/Online/OnlineSubsystemEOS/Source/OnlineSubsystemEOS/Public/Interfaces/OnlinePlayerReportEOSInterface.h b/Engine/Plugins/Online/OnlineSubsystemEOS/Source/OnlineSubsystemEOS/Public/Interfaces/OnlinePlayerReportEOSInterface.h
new file mode 100644
--- /dev/null
+++ b/Engine/Plugins/Online/OnlineSubsystemEOS/Source/OnlineSubsystemEOS/Public/Interfaces/OnlinePlayerReportEOSInterface.h
@@ -0,0 +1,61 @@
+// Copyright Epic Games, Inc. All Rights Reserved.
+
+#pragma once
+
+#include "Online/CoreOnline.h"
+#include "OnlineSubsystemTypes.h"
+#include "OnlineDelegateMacros.h"
+
+/**
+ * Delegate fired when the player report report was sent
+ *
+ *  @param bWasSuccessful true if the player report was successfully received
+ */
+DECLARE_DELEGATE_OneParam(FOnSendPlayerReportComplete, const bool /* bWasSuccessful */);
+
+/**
+ * Public Interface for interacting with EOS player reports
+ */
+class IOnlinePlayerReportEOS
+{
+public:
+	virtual ~IOnlinePlayerReportEOS() {};
+
+	/**
+	*	Enum to set the category on the player report
+	*/
+	enum class EPlayerReportCategory
+	{
+		Cheating,
+		Exploiting,
+		OffensiveProfile,
+		VerbalAbuse,
+		Scamming,
+		Spamming,
+		Other
+	};
+
+	/**
+	*	Struct for player report data
+	*/
+	struct FSendPlayerReportSettings
+	{
+		/** Required - The category of the player */
+		EPlayerReportCategory Category;
+
+		/** Optional - Message from the player descripting the issue being reported */
+		FString Message;
+
+		/** Optional - Context information around the issue being reported the game can send. This needs to be in a valid JSON format otherwise the EOS SDK will throw an error */
+		FString Context;
+	};
+
+	/**
+	* Send a player report for a player cheating, misconduct, etc... 
+	* @param LocalUserId the player id of the player sending the player report
+	* @param TargetUserId the player id the report is for
+	* @param FSendPlayerReportSettings the data to include in the player report.
+	*/
+	virtual void SendPlayerReport(const FUniqueNetId& LocalUserId, const FUniqueNetId& TargetUserId, FSendPlayerReportSettings&& SendPlayerReportSettings, FOnSendPlayerReportComplete&& Delegate) = 0;
+};
+typedef TSharedPtr<IOnlinePlayerReportEOS, ESPMode::ThreadSafe> IOnlinePlayerReportEOSPtr;
\ No newline at end of file
