Fix for override of the value the bUsesPresence in OSS EOS Lobbies

diff --git a/Engine/Plugins/Online/OnlineSubsystemEOS/Source/OnlineSubsystemEOS/Private/OnlineSessionEOS.cpp b/Engine/Plugins/Online/OnlineSubsystemEOS/Source/OnlineSubsystemEOS/Private/OnlineSessionEOS.cpp
--- a/Engine/Plugins/Online/OnlineSubsystemEOS/Source/OnlineSubsystemEOS/Private/OnlineSessionEOS.cpp
+++ b/Engine/Plugins/Online/OnlineSubsystemEOS/Source/OnlineSubsystemEOS/Private/OnlineSessionEOS.cpp
@@ -4855,10 +4855,10 @@ void FOnlineSessionEOS::CopyLobbyData(const TSharedRef<FLobbyDetailsEOS>& LobbyD
 {
 	// This method launches an asynchronous operation, so we'll pass the details handle as a shared ref to make sure it stays alive
 
+	// bUsesPresence will be set to false by default in search results, and it should be set by the game side before calling JoinSession.
+
 	OutSession.SessionSettings.bUseLobbiesIfAvailable = true;
 	OutSession.SessionSettings.bIsLANMatch = false;
-	 // By default, we'll set the search result's bUsesPresence to false. bUsesPresence should be set by the game side before calling JoinSession.
-	OutSession.SessionSettings.bUsesPresence = false;
 	OutSession.SessionSettings.Set(SETTING_HOST_MIGRATION, LobbyDetailsInfo->bAllowHostMigration, EOnlineDataAdvertisementType::DontAdvertise);
 #if WITH_EOS_RTC
 	OutSession.SessionSettings.bUseLobbiesVoiceChatIfAvailable = LobbyDetailsInfo->bRTCRoomEnabled == EOS_TRUE;
