Fix array RemoveAt call when that element wasn't set

diff --git a/Engine/Plugins/Online/OnlineSubsystemEOS/Source/OnlineSubsystemEOS/Private/UserManagerEOS.cpp b/Engine/Plugins/Online/OnlineSubsystemEOS/Source/OnlineSubsystemEOS/Private/UserManagerEOS.cpp
--- a/Engine/Plugins/Online/OnlineSubsystemEOS/Source/OnlineSubsystemEOS/Private/UserManagerEOS.cpp
+++ b/Engine/Plugins/Online/OnlineSubsystemEOS/Source/OnlineSubsystemEOS/Private/UserManagerEOS.cpp
@@ -589,6 +589,14 @@ FString FUserManagerEOS::GetEOSAuthTokenFilename()
 	return EOSSubsystem->ProductId + TEXT("_") + EOS_EPIC_AUTH_TOKEN_FILENAME_SUFFIX;
 }
 
+void FUserManagerEOS::TryRemoveLoginSession(int32 LocalUserNum)
+{
+	if (LoginSessions.IsValidIndex(LocalUserNum))
+	{
+		LoginSessions.RemoveAt(LocalUserNum);
+	}
+}
+
 void FUserManagerEOS::LoginViaPersistentAuthLegacy(int32 LocalUserNum, const FOnlineAccountCredentials& AccountCredentials)
 {
 	FPlatformEOSHelpersPtr EOSHelpers = EOSSubsystem->GetEOSHelpers();
@@ -748,7 +756,7 @@ void FUserManagerEOS::LoginViaExternalAuth(int32 LocalUserNum)
 					if (!bWasSuccessful || !AuthToken.IsValid())
 					{
 						UE_LOG_ONLINE(Warning, TEXT("Unable to Login() user (%d) due to an empty platform auth token"), LocalUserNum);
-						LoginSessions.RemoveAt(LocalUserNum);
+						TryRemoveLoginSession(LocalUserNum);
 						TriggerOnLoginCompleteDelegates(LocalUserNum, false, *FUniqueNetIdEOS::EmptyId(), FString(TEXT("Missing platform auth token")));
 						return;
 					}
@@ -767,7 +775,7 @@ void FUserManagerEOS::LoginViaExternalAuth(int32 LocalUserNum)
 					else
 					{
 						UE_LOG_ONLINE(Error, TEXT("FAuthCredentials object cannot be constructed with invalid FExternalAuthToken parameter"));
-						LoginSessions.RemoveAt(LocalUserNum);
+						TryRemoveLoginSession(LocalUserNum);
 						TriggerOnLoginCompleteDelegates(LocalUserNum, false, *FUniqueNetIdEOS::EmptyId(), FString(TEXT("Invalid platform auth token")));
 						return;
 					}
@@ -827,7 +835,7 @@ void FUserManagerEOS::CallEOSAuthLogin(int32 LocalUserNum, const FOnlineAccountC
 	if (!ToEOS_ELoginCredentialType(CredentialsType, EOSCredentials.Type))
 	{
 		UE_LOG_ONLINE(Warning, TEXT("Unable to Login() user (%d) due to missing auth parameters"), LocalUserNum);
-		LoginSessions.RemoveAt(LocalUserNum);
+		TryRemoveLoginSession(LocalUserNum);
 		TriggerOnLoginCompleteDelegates(LocalUserNum, false, *FUniqueNetIdEOS::EmptyId(), FString(TEXT("Missing auth parameters")));
 		return;
 	}
@@ -851,7 +859,7 @@ void FUserManagerEOS::CallEOSAuthLogin(int32 LocalUserNum, const FOnlineAccountC
 				if (!LexFromString(EOSCredentials.ExternalType, *CredentialsTokenType))
 				{
 					UE_LOG_ONLINE(Warning, TEXT("Unable to Login() user (%d). External Auth Token Type not valid."), LocalUserNum);
-					LoginSessions.RemoveAt(LocalUserNum);
+					TryRemoveLoginSession(LocalUserNum);
 					TriggerOnLoginCompleteDelegates(LocalUserNum, false, *FUniqueNetIdEOS::EmptyId(), FString(TEXT("External Auth Token Type not valid")));
 					return;
 				}
@@ -859,7 +867,7 @@ void FUserManagerEOS::CallEOSAuthLogin(int32 LocalUserNum, const FOnlineAccountC
 			else
 			{
 				UE_LOG_ONLINE(Warning, TEXT("Unable to Login() user (%d). External Auth Token Type not specified."), LocalUserNum);
-				LoginSessions.RemoveAt(LocalUserNum);
+				TryRemoveLoginSession(LocalUserNum);
 				TriggerOnLoginCompleteDelegates(LocalUserNum, false, *FUniqueNetIdEOS::EmptyId(), FString(TEXT("External Auth Token Type not specified")));
 				return;
 			}
@@ -1064,7 +1072,7 @@ void FUserManagerEOS::OnEOSAuthLoginComplete(int32 LocalUserNum, const EOS_ELogi
 		{
 			FString ErrorString = FString::Printf(TEXT("Login(%d) failed with EOS result code (%s)"), LocalUserNum, *LexToString(LoginResultCode));
 			UE_LOG_ONLINE(Warning, TEXT("%s"), *ErrorString);
-			LoginSessions.RemoveAt(LocalUserNum);
+			TryRemoveLoginSession(LocalUserNum);
 			TriggerOnLoginCompleteDelegates(LocalUserNum, false, *FUniqueNetIdEOS::EmptyId(), ErrorString);
 		};
 
@@ -1160,7 +1168,7 @@ void FUserManagerEOS::LinkEAS(int32 LocalUserNum, EOS_ContinuanceToken Token)
 			{
 				FString ErrorString = FString::Printf(TEXT("Login(%d) failed with EOS result code (%s)"), LocalUserNum, *LexToString(Data->ResultCode));
 				UE_LOG_ONLINE(Warning, TEXT("%s"), *ErrorString);
-				LoginSessions.RemoveAt(LocalUserNum);
+				TryRemoveLoginSession(LocalUserNum);
 				TriggerOnLoginCompleteDelegates(LocalUserNum, false, *FUniqueNetIdEOS::EmptyId(), ErrorString);
 			}
 		};
@@ -1225,7 +1233,7 @@ void FUserManagerEOS::CallEOSConnectLogin(int32 LocalUserNum, const FOnlineAccou
 			if (!LexFromString(ExternalCredentialType, *CredentialsTokenType))
 			{
 				UE_LOG_ONLINE(Warning, TEXT("Unable to Login() user (%d). External Auth Token Type not valid."), LocalUserNum);
-				LoginSessions.RemoveAt(LocalUserNum);
+				TryRemoveLoginSession(LocalUserNum);
 				TriggerOnLoginCompleteDelegates(LocalUserNum, false, *FUniqueNetIdEOS::EmptyId(), FString(TEXT("External Auth Token Type not valid")));
 				return;
 			}
@@ -1233,7 +1241,7 @@ void FUserManagerEOS::CallEOSConnectLogin(int32 LocalUserNum, const FOnlineAccou
 		else
 		{
 			UE_LOG_ONLINE(Warning, TEXT("Unable to Login() user (%d). External Auth Token Type not specified."), LocalUserNum);
-			LoginSessions.RemoveAt(LocalUserNum);
+			TryRemoveLoginSession(LocalUserNum);
 			TriggerOnLoginCompleteDelegates(LocalUserNum, false, *FUniqueNetIdEOS::EmptyId(), FString(TEXT("External Auth Token Type not specified")));
 			return;
 		}
@@ -1279,7 +1287,7 @@ void FUserManagerEOS::CallEOSConnectLogin(int32 LocalUserNum, const FOnlineAccou
 			{
 				const FString ErrorString = FString::Printf(TEXT("ConnectLoginNoEAS(%d) failed with EOS result code (%s)"), LocalUserNum, *LexToString(Data->ResultCode));
 				UE_LOG_ONLINE(Warning, TEXT("%s"), *ErrorString);
-				LoginSessions.RemoveAt(LocalUserNum);
+				TryRemoveLoginSession(LocalUserNum);
 				TriggerOnLoginCompleteDelegates(LocalUserNum, false, *FUniqueNetIdEOS::EmptyId(), ErrorString);
 			}
 		};
@@ -1321,7 +1329,7 @@ bool FUserManagerEOS::ConnectLoginNoEAS(int32 LocalUserNum)
 					{
 						const FString ErrorString = FString::Printf(TEXT("ConnectLoginNoEAS(%d) failed due to the platform OSS giving an empty auth token"), LocalUserNum);
 						UE_LOG_ONLINE(Warning, TEXT("%s"), *ErrorString);
-						LoginSessions.RemoveAt(LocalUserNum);
+						TryRemoveLoginSession(LocalUserNum);
 						TriggerOnLoginCompleteDelegates(LocalUserNum, false, *FUniqueNetIdEOS::EmptyId(), ErrorString);
 						return;
 					}
@@ -1432,7 +1440,7 @@ bool FUserManagerEOS::ConnectLoginEAS(int32 LocalUserNum, EOS_EpicAccountId Acco
 					{
 						const FString ErrorString = FString::Printf(TEXT("ConnectLoginEAS(%d) failed due to the platform OSS giving an empty auth token"), LocalUserNum);
 						UE_LOG_ONLINE(Warning, TEXT("%s"), *ErrorString);
-						LoginSessions.RemoveAt(LocalUserNum);
+						TryRemoveLoginSession(LocalUserNum);
 						TriggerOnLoginCompleteDelegates(LocalUserNum, false, *FUniqueNetIdEOS::EmptyId(), ErrorString);
 						return;
 					}
@@ -1683,7 +1691,7 @@ void FUserManagerEOS::CreateConnectedLogin(int32 LocalUserNum, EOS_EpicAccountId
 			{
 				// @todo joeg - logout?
 				FString ErrorString = FString::Printf(TEXT("Login(%d) failed with EOS result code (%s)"), LocalUserNum, *LexToString(Data->ResultCode));
-				LoginSessions.RemoveAt(LocalUserNum);
+				TryRemoveLoginSession(LocalUserNum);
 				TriggerOnLoginCompleteDelegates(LocalUserNum, false, *FUniqueNetIdEOS::EmptyId(), ErrorString);
 			}
 		};
@@ -1785,7 +1793,7 @@ bool FUserManagerEOS::Logout(int32 LocalUserNum)
 			{
 				RemoveLocalUser(LocalUserNum);
 
-				LoginSessions.RemoveAt(LocalUserNum);
+				TryRemoveLoginSession(LocalUserNum);
 
 				TriggerOnLogoutCompleteDelegates(LocalUserNum, true);
 			}
diff --git a/Engine/Plugins/Online/OnlineSubsystemEOS/Source/OnlineSubsystemEOS/Private/UserManagerEOS.h b/Engine/Plugins/Online/OnlineSubsystemEOS/Source/OnlineSubsystemEOS/Private/UserManagerEOS.h
--- a/Engine/Plugins/Online/OnlineSubsystemEOS/Source/OnlineSubsystemEOS/Private/UserManagerEOS.h
+++ b/Engine/Plugins/Online/OnlineSubsystemEOS/Source/OnlineSubsystemEOS/Private/UserManagerEOS.h
@@ -470,6 +470,8 @@ private:
 
 	/** Relevant information about the login session for local users */
 	TLocalUserArray<FLoginSession> LoginSessions;
+	// Remove from LoginSessions, but only if LocalUserNum is in there
+	void TryRemoveLoginSession(int32 LocalUserNum);
 
 	// Online user info
 
