[Backout] - CL33507255
- Backing out as this situation should be impossible, due to module load/unload order guarantees. Something has broken those guarantees for the affected licensee but this is not the fix.
Original CL Desc
-----------------------------------------------------------------
PR #11779: Don't attempt a tick at shutdown if EOSShared module is no longer loaded.

diff --git a/Engine/Plugins/Online/OnlineSubsystemEOS/Source/OnlineSubsystemEOS/Private/OnlineSubsystemEOS.cpp b/Engine/Plugins/Online/OnlineSubsystemEOS/Source/OnlineSubsystemEOS/Private/OnlineSubsystemEOS.cpp
--- a/Engine/Plugins/Online/OnlineSubsystemEOS/Source/OnlineSubsystemEOS/Private/OnlineSubsystemEOS.cpp
+++ b/Engine/Plugins/Online/OnlineSubsystemEOS/Source/OnlineSubsystemEOS/Private/OnlineSubsystemEOS.cpp
@@ -446,12 +446,7 @@ bool FOnlineSubsystemEOS::Shutdown()
 	// EOS-22677 workaround: Make sure tick is called at least once before shutting down.
 	if (EOSPlatformHandle)
 	{
-		// The EOSShared module may have been shut down at this point, in which case
-		// the handle has been released and we can't use it.
-		if (FModuleManager::Get().IsModuleLoaded("EOSShared"))
-		{
-			EOS_Platform_Tick(*EOSPlatformHandle);
-		}
+		EOS_Platform_Tick(*EOSPlatformHandle);
 	}
 
 	StopTicker();
