Addition of support for new result of EOS_Ecom_Checkout: EOS_Ecom_PurchaseProcessing in OSS.
Fix for Checkout flow in OnlineServicesEOS

diff --git a/Engine/Plugins/Online/OnlineServicesEOS/Source/Private/Online/CommerceEOS.cpp b/Engine/Plugins/Online/OnlineServicesEOS/Source/Private/Online/CommerceEOS.cpp
--- a/Engine/Plugins/Online/OnlineServicesEOS/Source/Private/Online/CommerceEOS.cpp
+++ b/Engine/Plugins/Online/OnlineServicesEOS/Source/Private/Online/CommerceEOS.cpp
@@ -224,10 +224,10 @@ TOnlineAsyncOpHandle<FCommerceCheckout> FCommerceEOS::Checkout(FCommerceCheckout
 		{
 			Utf8CheckoutIds.Emplace(*Offer.OfferId);
 
-			EosCheckoutEntries.AddDefaulted();
-			EosCheckoutEntries.Last().ApiVersion = 1;
+			EOS_Ecom_CheckoutEntry& CheckoutEntry = EosCheckoutEntries.AddDefaulted_GetRef();
+			CheckoutEntry.ApiVersion = 1;
 			UE_EOS_CHECK_API_MISMATCH(EOS_ECOM_CHECKOUTENTRY_API_LATEST, 1);
-			EosCheckoutEntries.Last().OfferId = Utf8CheckoutIds.Last().Get();
+			CheckoutEntry.OfferId = Utf8CheckoutIds.Last().Get();
 		}
 
 		EOS_Ecom_CheckoutOptions Options = { };
diff --git a/Engine/Plugins/Online/OnlineSubsystemEOS/Source/OnlineSubsystemEOS/Private/OnlineStoreEOS.cpp b/Engine/Plugins/Online/OnlineSubsystemEOS/Source/OnlineSubsystemEOS/Private/OnlineStoreEOS.cpp
--- a/Engine/Plugins/Online/OnlineSubsystemEOS/Source/OnlineSubsystemEOS/Private/OnlineStoreEOS.cpp
+++ b/Engine/Plugins/Online/OnlineSubsystemEOS/Source/OnlineSubsystemEOS/Private/OnlineStoreEOS.cpp
@@ -198,11 +198,15 @@ void FOnlineStoreEOS::Checkout(const FUniqueNetId& UserId, const FPurchaseChecko
 		const EOS_EResult CheckoutResult = Data->ResultCode;
 		if (CheckoutResult != EOS_EResult::EOS_Success)
 		{
-			UE_LOG_ONLINE(Error, TEXT("EOS_Ecom_Checkout: failed with error (%s)"), *LexToString(Data->ResultCode));
+			UE_LOG_ONLINE(Error, TEXT("EOS_Ecom_Checkout: failed with error (%s)"), *LexToString(CheckoutResult));
 			if (CheckoutResult == EOS_EResult::EOS_Canceled)
 			{
 				OnComplete.ExecuteIfBound(ONLINE_ERROR(EOnlineErrorResult::Canceled), MakeShared<FPurchaseReceipt>());
 			}
+			else if (CheckoutResult == EOS_EResult::EOS_Ecom_PurchaseProcessing)
+			{
+				OnComplete.ExecuteIfBound(ONLINE_ERROR(EOnlineErrorResult::FailExtended, LexToString(CheckoutResult)), MakeShared<FPurchaseReceipt>());
+			}
 			else
 			{
 				OnComplete.ExecuteIfBound(ONLINE_ERROR(EOnlineErrorResult::Unknown), MakeShared<FPurchaseReceipt>());
