Added missing validity check in EOS user info callback.

diff --git a/Engine/Plugins/Online/OnlineSubsystemEOS/Source/OnlineSubsystemEOS/Private/UserManagerEOS.cpp b/Engine/Plugins/Online/OnlineSubsystemEOS/Source/OnlineSubsystemEOS/Private/UserManagerEOS.cpp
--- a/Engine/Plugins/Online/OnlineSubsystemEOS/Source/OnlineSubsystemEOS/Private/UserManagerEOS.cpp
+++ b/Engine/Plugins/Online/OnlineSubsystemEOS/Source/OnlineSubsystemEOS/Private/UserManagerEOS.cpp
@@ -3923,8 +3923,11 @@ void FUserManagerEOS::ReadUserInfo(int32 LocalUserNum, EOS_EpicAccountId EpicAcc
 			ErrorStr = LexToString(Data->ResultCode);
 		}
 
-		// We mark this player as processed
-		GetLocalUserChecked(LocalUserNum).OngoingQueryUserInfoAccounts.RemoveSwap(EpicAccountId, EAllowShrinking::No);
+		// We mark this player as processed if it exists
+		if(IsLocalUserValid(LocalUserNum))
+		{
+			GetLocalUserChecked(LocalUserNum).OngoingQueryUserInfoAccounts.RemoveSwap(EpicAccountId, EAllowShrinking::No);
+		}
 
 		Callback(bWasSuccessful, EOSId, ErrorStr);
 	};
