Refactored OSS EOS Login flow to remove the need for FLoginSession type.
Local user is now added at the beginning of the flow and updated at the end.

diff --git a/Engine/Plugins/Online/OnlineSubsystemEOS/Source/OnlineSubsystemEOS/Private/UserManagerEOS.cpp b/Engine/Plugins/Online/OnlineSubsystemEOS/Source/OnlineSubsystemEOS/Private/UserManagerEOS.cpp
--- a/Engine/Plugins/Online/OnlineSubsystemEOS/Source/OnlineSubsystemEOS/Private/UserManagerEOS.cpp
+++ b/Engine/Plugins/Online/OnlineSubsystemEOS/Source/OnlineSubsystemEOS/Private/UserManagerEOS.cpp
@@ -540,10 +540,7 @@ bool FUserManagerEOS::Login(int32 LocalUserNum, const FOnlineAccountCredentials&
 		return true;
 	}
 
-	if(!LoginSessions.IsValidIndex(LocalUserNum))
-	{
-		LoginSessions.Insert(LocalUserNum, FLoginSession());
-	}
+	FLocalUserEOS& LocalUser = AddLocalUser(LocalUserNum);
 
 	// EOS Connect without EAS
 	if(!Settings.bUseEAS && Settings.bUseEOSConnect)
@@ -551,7 +548,7 @@ bool FUserManagerEOS::Login(int32 LocalUserNum, const FOnlineAccountCredentials&
 		if(!AccountCredentials.Token.IsEmpty())
 		{
 			// We'll save that a Connect Token was manually provided for this local user
-			LoginSessions[LocalUserNum].UserProvidedConnectAuthToken = AccountCredentials.Token;
+			LocalUser.UserProvidedConnectAuthToken = AccountCredentials.Token;
 		}
 
 		CallEOSConnectLogin(LocalUserNum, AccountCredentials);
@@ -589,14 +586,6 @@ FString FUserManagerEOS::GetEOSAuthTokenFilename()
 	return EOSSubsystem->ProductId + TEXT("_") + EOS_EPIC_AUTH_TOKEN_FILENAME_SUFFIX;
 }
 
-void FUserManagerEOS::TryRemoveLoginSession(int32 LocalUserNum)
-{
-	if (LoginSessions.IsValidIndex(LocalUserNum))
-	{
-		LoginSessions.RemoveAt(LocalUserNum);
-	}
-}
-
 void FUserManagerEOS::LoginViaPersistentAuthLegacy(int32 LocalUserNum, const FOnlineAccountCredentials& AccountCredentials)
 {
 	FPlatformEOSHelpersPtr EOSHelpers = EOSSubsystem->GetEOSHelpers();
@@ -756,7 +745,7 @@ void FUserManagerEOS::LoginViaExternalAuth(int32 LocalUserNum)
 					if (!bWasSuccessful || !AuthToken.IsValid())
 					{
 						UE_LOG_ONLINE(Warning, TEXT("Unable to Login() user (%d) due to an empty platform auth token"), LocalUserNum);
-						TryRemoveLoginSession(LocalUserNum);
+						RemoveLocalUser(LocalUserNum);
 						TriggerOnLoginCompleteDelegates(LocalUserNum, false, *FUniqueNetIdEOS::EmptyId(), FString(TEXT("Missing platform auth token")));
 						return;
 					}
@@ -775,7 +764,7 @@ void FUserManagerEOS::LoginViaExternalAuth(int32 LocalUserNum)
 					else
 					{
 						UE_LOG_ONLINE(Error, TEXT("FAuthCredentials object cannot be constructed with invalid FExternalAuthToken parameter"));
-						TryRemoveLoginSession(LocalUserNum);
+						RemoveLocalUser(LocalUserNum);
 						TriggerOnLoginCompleteDelegates(LocalUserNum, false, *FUniqueNetIdEOS::EmptyId(), FString(TEXT("Invalid platform auth token")));
 						return;
 					}
@@ -835,7 +824,7 @@ void FUserManagerEOS::CallEOSAuthLogin(int32 LocalUserNum, const FOnlineAccountC
 	if (!ToEOS_ELoginCredentialType(CredentialsType, EOSCredentials.Type))
 	{
 		UE_LOG_ONLINE(Warning, TEXT("Unable to Login() user (%d) due to missing auth parameters"), LocalUserNum);
-		TryRemoveLoginSession(LocalUserNum);
+		RemoveLocalUser(LocalUserNum);
 		TriggerOnLoginCompleteDelegates(LocalUserNum, false, *FUniqueNetIdEOS::EmptyId(), FString(TEXT("Missing auth parameters")));
 		return;
 	}
@@ -859,7 +848,7 @@ void FUserManagerEOS::CallEOSAuthLogin(int32 LocalUserNum, const FOnlineAccountC
 				if (!LexFromString(EOSCredentials.ExternalType, *CredentialsTokenType))
 				{
 					UE_LOG_ONLINE(Warning, TEXT("Unable to Login() user (%d). External Auth Token Type not valid."), LocalUserNum);
-					TryRemoveLoginSession(LocalUserNum);
+					RemoveLocalUser(LocalUserNum);
 					TriggerOnLoginCompleteDelegates(LocalUserNum, false, *FUniqueNetIdEOS::EmptyId(), FString(TEXT("External Auth Token Type not valid")));
 					return;
 				}
@@ -867,7 +856,7 @@ void FUserManagerEOS::CallEOSAuthLogin(int32 LocalUserNum, const FOnlineAccountC
 			else
 			{
 				UE_LOG_ONLINE(Warning, TEXT("Unable to Login() user (%d). External Auth Token Type not specified."), LocalUserNum);
-				TryRemoveLoginSession(LocalUserNum);
+				RemoveLocalUser(LocalUserNum);
 				TriggerOnLoginCompleteDelegates(LocalUserNum, false, *FUniqueNetIdEOS::EmptyId(), FString(TEXT("External Auth Token Type not specified")));
 				return;
 			}
@@ -1072,7 +1061,7 @@ void FUserManagerEOS::OnEOSAuthLoginComplete(int32 LocalUserNum, const EOS_ELogi
 		{
 			FString ErrorString = FString::Printf(TEXT("Login(%d) failed with EOS result code (%s)"), LocalUserNum, *LexToString(LoginResultCode));
 			UE_LOG_ONLINE(Warning, TEXT("%s"), *ErrorString);
-			TryRemoveLoginSession(LocalUserNum);
+			RemoveLocalUser(LocalUserNum);
 			TriggerOnLoginCompleteDelegates(LocalUserNum, false, *FUniqueNetIdEOS::EmptyId(), ErrorString);
 		};
 
@@ -1168,7 +1157,7 @@ void FUserManagerEOS::LinkEAS(int32 LocalUserNum, EOS_ContinuanceToken Token)
 			{
 				FString ErrorString = FString::Printf(TEXT("Login(%d) failed with EOS result code (%s)"), LocalUserNum, *LexToString(Data->ResultCode));
 				UE_LOG_ONLINE(Warning, TEXT("%s"), *ErrorString);
-				TryRemoveLoginSession(LocalUserNum);
+				RemoveLocalUser(LocalUserNum);
 				TriggerOnLoginCompleteDelegates(LocalUserNum, false, *FUniqueNetIdEOS::EmptyId(), ErrorString);
 			}
 		};
@@ -1233,7 +1222,7 @@ void FUserManagerEOS::CallEOSConnectLogin(int32 LocalUserNum, const FOnlineAccou
 			if (!LexFromString(ExternalCredentialType, *CredentialsTokenType))
 			{
 				UE_LOG_ONLINE(Warning, TEXT("Unable to Login() user (%d). External Auth Token Type not valid."), LocalUserNum);
-				TryRemoveLoginSession(LocalUserNum);
+				RemoveLocalUser(LocalUserNum);
 				TriggerOnLoginCompleteDelegates(LocalUserNum, false, *FUniqueNetIdEOS::EmptyId(), FString(TEXT("External Auth Token Type not valid")));
 				return;
 			}
@@ -1241,7 +1230,7 @@ void FUserManagerEOS::CallEOSConnectLogin(int32 LocalUserNum, const FOnlineAccou
 		else
 		{
 			UE_LOG_ONLINE(Warning, TEXT("Unable to Login() user (%d). External Auth Token Type not specified."), LocalUserNum);
-			TryRemoveLoginSession(LocalUserNum);
+			RemoveLocalUser(LocalUserNum);
 			TriggerOnLoginCompleteDelegates(LocalUserNum, false, *FUniqueNetIdEOS::EmptyId(), FString(TEXT("External Auth Token Type not specified")));
 			return;
 		}
@@ -1287,7 +1276,7 @@ void FUserManagerEOS::CallEOSConnectLogin(int32 LocalUserNum, const FOnlineAccou
 			{
 				const FString ErrorString = FString::Printf(TEXT("ConnectLoginNoEAS(%d) failed with EOS result code (%s)"), LocalUserNum, *LexToString(Data->ResultCode));
 				UE_LOG_ONLINE(Warning, TEXT("%s"), *ErrorString);
-				TryRemoveLoginSession(LocalUserNum);
+				RemoveLocalUser(LocalUserNum);
 				TriggerOnLoginCompleteDelegates(LocalUserNum, false, *FUniqueNetIdEOS::EmptyId(), ErrorString);
 			}
 		};
@@ -1329,7 +1318,7 @@ bool FUserManagerEOS::ConnectLoginNoEAS(int32 LocalUserNum)
 					{
 						const FString ErrorString = FString::Printf(TEXT("ConnectLoginNoEAS(%d) failed due to the platform OSS giving an empty auth token"), LocalUserNum);
 						UE_LOG_ONLINE(Warning, TEXT("%s"), *ErrorString);
-						TryRemoveLoginSession(LocalUserNum);
+						RemoveLocalUser(LocalUserNum);
 						TriggerOnLoginCompleteDelegates(LocalUserNum, false, *FUniqueNetIdEOS::EmptyId(), ErrorString);
 						return;
 					}
@@ -1440,7 +1429,7 @@ bool FUserManagerEOS::ConnectLoginEAS(int32 LocalUserNum, EOS_EpicAccountId Acco
 					{
 						const FString ErrorString = FString::Printf(TEXT("ConnectLoginEAS(%d) failed due to the platform OSS giving an empty auth token"), LocalUserNum);
 						UE_LOG_ONLINE(Warning, TEXT("%s"), *ErrorString);
-						TryRemoveLoginSession(LocalUserNum);
+						RemoveLocalUser(LocalUserNum);
 						TriggerOnLoginCompleteDelegates(LocalUserNum, false, *FUniqueNetIdEOS::EmptyId(), ErrorString);
 						return;
 					}
@@ -1570,17 +1559,17 @@ void FUserManagerEOS::AutoRefreshConnectLogin(int32 LocalUserNum)
 	}
 	else
 	{
-		if (LoginSessions.IsValidIndex(LocalUserNum))
+		if (IsLocalUserValid(LocalUserNum))
 		{
-			const FLoginSession& LoginSession = LoginSessions[LocalUserNum];
+			const FLocalUserEOS& LocalUser = GetLocalUserChecked(LocalUserNum);
 
 			// If the user provided a Connect Auth Token, we'll try renewing the auth with it before we fire the corresponding delegate
-			if (LoginSession.UserProvidedConnectAuthToken.IsSet())
+			if (LocalUser.UserProvidedConnectAuthToken.IsSet())
 			{
 				IOnlineSubsystem* PlatformOSS = GetPlatformOSS();
 				EOS_EExternalCredentialType CredType = ToEOS_EExternalCredentialType(PlatformOSS ? PlatformOSS->GetSubsystemName() : EOSSubsystem->GetSubsystemName());
 
-				FConnectCredentials Credentials(CredType, *LoginSession.UserProvidedConnectAuthToken);
+				FConnectCredentials Credentials(CredType, *LocalUser.UserProvidedConnectAuthToken);
 
 				EOS_Connect_LoginOptions Options = { };
 				Options.ApiVersion = 2;
@@ -1691,7 +1680,7 @@ void FUserManagerEOS::CreateConnectedLogin(int32 LocalUserNum, EOS_EpicAccountId
 			{
 				// @todo joeg - logout?
 				FString ErrorString = FString::Printf(TEXT("Login(%d) failed with EOS result code (%s)"), LocalUserNum, *LexToString(Data->ResultCode));
-				TryRemoveLoginSession(LocalUserNum);
+				RemoveLocalUser(LocalUserNum);
 				TriggerOnLoginCompleteDelegates(LocalUserNum, false, *FUniqueNetIdEOS::EmptyId(), ErrorString);
 			}
 		};
@@ -1764,7 +1753,7 @@ void FUserManagerEOS::FullLoginCallback(int32 LocalUserNum, EOS_EpicAccountId Ac
 	}
 
 	// We add the local user
-	FLocalUserEOS& LocalUser = AddLocalUser(LocalUserNum, AccountId, UserId);
+	FLocalUserEOS& LocalUser = UpdateLocalUser(LocalUserNum, AccountId, UserId);
 
 	TriggerOnLoginCompleteDelegates(LocalUserNum, true, *LocalUser.UniqueNetId, FString());
 	TriggerOnLoginStatusChangedDelegates(LocalUserNum, ELoginStatus::NotLoggedIn, ELoginStatus::LoggedIn, *LocalUser.UniqueNetId);
@@ -1793,8 +1782,6 @@ bool FUserManagerEOS::Logout(int32 LocalUserNum)
 			{
 				RemoveLocalUser(LocalUserNum);
 
-				TryRemoveLoginSession(LocalUserNum);
-
 				TriggerOnLogoutCompleteDelegates(LocalUserNum, true);
 			}
 			else
@@ -1893,10 +1880,7 @@ bool FUserManagerEOS::AutoLogin(int32 LocalUserNum)
 		return true;
 	}
 
-	if (!LoginSessions.IsValidIndex(LocalUserNum))
-	{
-		LoginSessions.Insert(LocalUserNum, FLoginSession());
-	}
+	AddLocalUser(LocalUserNum);
 
 	// Check if we are configured to prefer persistent auth when attempting EAS login
 	if (Settings.bUseEAS && Settings.bPreferPersistentAuth)
@@ -1923,19 +1907,8 @@ bool FUserManagerEOS::AutoLogin(int32 LocalUserNum)
 	return false;
 }
 
-FLocalUserEOS& FUserManagerEOS::AddLocalUser(int32 LocalUserNum, EOS_EpicAccountId EpicAccountId, EOS_ProductUserId UserId)
+FLocalUserEOS& FUserManagerEOS::UpdateLocalUser(int32 LocalUserNum, EOS_EpicAccountId EpicAccountId, EOS_ProductUserId UserId)
 {
-	// Set the default user to the first one that logs in
-	if (DefaultLocalUser == INVALID_LOCAL_USER)
-	{
-		DefaultLocalUser = LocalUserNum;
-	}
-
-	// Init player lists
-	if (!LocalUsers.IsValidIndex(LocalUserNum))
-	{
-		LocalUsers.Insert(LocalUserNum, FLocalUserEOS());
-	}
 	FLocalUserEOS& LocalUser = GetLocalUserChecked(LocalUserNum);
 
 	FUniqueNetIdEOSRef UserNetId = FUniqueNetIdEOSRegistry::FindOrAdd(EpicAccountId, UserId);
@@ -2458,15 +2431,33 @@ FOnlineUserPtr FUserManagerEOS::GetOnlineUser(EOS_EpicAccountId AccountId) const
 	return OnlineUser;
 }
 
+FLocalUserEOS& FUserManagerEOS::AddLocalUser(int32 LocalUserNum)
+{
+	// Set the default user to the first one that logs in
+	if (DefaultLocalUser == INVALID_LOCAL_USER)
+	{
+		DefaultLocalUser = LocalUserNum;
+	}
+
+	// Init player lists
+	if (!LocalUsers.IsValidIndex(LocalUserNum))
+	{
+		LocalUsers.Insert(LocalUserNum, FLocalUserEOS());
+	}
+
+	return GetLocalUserChecked(LocalUserNum);
+}
+
 void FUserManagerEOS::RemoveLocalUser(int32 LocalUserNum)
 {
 	if (LocalUsers.IsValidIndex(LocalUserNum))
 	{
 		FLocalUserEOS& LocalUser = GetLocalUserChecked(LocalUserNum);
-		const FUniqueNetIdEOSRef FoundId = LocalUser.UniqueNetId.ToSharedRef();
-
-		EOSSubsystem->ReleaseVoiceChatUserInterface(*FoundId);
 
+		if(LocalUser.UniqueNetId.IsValid())
+		{
+			EOSSubsystem->ReleaseVoiceChatUserInterface(*LocalUser.UniqueNetId);
+		}
 
 		for (const ReadUserListInfo& CachedInfo : LocalUser.CachedReadUserListInfo)
 		{
diff --git a/Engine/Plugins/Online/OnlineSubsystemEOS/Source/OnlineSubsystemEOS/Private/UserManagerEOS.h b/Engine/Plugins/Online/OnlineSubsystemEOS/Source/OnlineSubsystemEOS/Private/UserManagerEOS.h
--- a/Engine/Plugins/Online/OnlineSubsystemEOS/Source/OnlineSubsystemEOS/Private/UserManagerEOS.h
+++ b/Engine/Plugins/Online/OnlineSubsystemEOS/Source/OnlineSubsystemEOS/Private/UserManagerEOS.h
@@ -191,7 +191,7 @@ struct ReadUserListInfo
 	};
 };
 
-/** Information related to ongoing operations and state for local users */
+/** Created when a new local user starts a login process, keeps the state for the login session and ongoing operations, and gets destroyed when login fails or the user logs out */
 struct FLocalUserEOS
 {
 	/** UniqueNetId for the local user, set after successful authentication */
@@ -236,11 +236,7 @@ struct FLocalUserEOS
 
 	/** Struct containing the information related to the connect login notification for the user */
 	FNotificationIdCallbackPairPtr ConnectLoginNotification;
-};
 
-/** Created when a new local user starts a login process, keeps the state for the login session, and gets destroyed when the user logs out */
-struct FLoginSession
-{
 	/** If set, Connect Login will be re-attempted with it. If this attempt fails, a new token will need to be provided */
 	TOptional<FString> UserProvidedConnectAuthToken;
 };
@@ -424,8 +420,9 @@ private:
 	void OnEOSAuthLoginComplete(int32 LocalUserNum, const EOS_ELoginCredentialType LoginCredentialType, const bool bIsAutoLogin, const EOS_Auth_LoginCallbackInfo* Data);
 	void CallEOSConnectLogin(int32 LocalUserNum, const FOnlineAccountCredentials& Credentials);
 
+	FLocalUserEOS& AddLocalUser(int32 LocalUserNum);
+	FLocalUserEOS& UpdateLocalUser(int32 LocalUserNum, EOS_EpicAccountId EpicAccountId, EOS_ProductUserId UserId);
 	void RemoveLocalUser(int32 LocalUserNum);
-	FLocalUserEOS& AddLocalUser(int32 LocalUserNum, EOS_EpicAccountId EpicAccountId, EOS_ProductUserId UserId);
 
 	typedef TFunction<void(bool bWasSuccessful, FUniqueNetIdEOSRef RemotePlayerNetId, const FString& ErrorStr)> FRemoteUserProcessedCallback;
 	void AddRemotePlayer(int32 LocalUserNum, EOS_EpicAccountId EpicAccountId, const FRemoteUserProcessedCallback& Callback);
@@ -468,11 +465,6 @@ private:
 	/** Information related to ongoing operations and state for local users */
 	TLocalUserArray<FLocalUserEOS> LocalUsers;
 
-	/** Relevant information about the login session for local users */
-	TLocalUserArray<FLoginSession> LoginSessions;
-	// Remove from LoginSessions, but only if LocalUserNum is in there
-	void TryRemoveLoginSession(int32 LocalUserNum);
-
 	// Online user info
 
 	/** Ids mapped to attribute access information for any user */
