OSS EOS Legacy login flow fix

diff --git a/Engine/Plugins/Online/OnlineSubsystemEOS/Source/OnlineSubsystemEOS/Private/UserManagerEOS.cpp b/Engine/Plugins/Online/OnlineSubsystemEOS/Source/OnlineSubsystemEOS/Private/UserManagerEOS.cpp
--- a/Engine/Plugins/Online/OnlineSubsystemEOS/Source/OnlineSubsystemEOS/Private/UserManagerEOS.cpp
+++ b/Engine/Plugins/Online/OnlineSubsystemEOS/Source/OnlineSubsystemEOS/Private/UserManagerEOS.cpp
@@ -1694,6 +1694,9 @@ typedef TEOSGlobalCallback<EOS_Auth_OnLoginStatusChangedCallback, EOS_Auth_Login
 
 void FUserManagerEOS::FullLoginCallbackLegacy(int32 LocalUserNum, EOS_EpicAccountId AccountId, EOS_ProductUserId UserId, const FOnlineAccountCredentials& AccountCredentials)
 {
+	// Before we call the method from the non-legacy flow, we must add the local user, as this flow counts on it already existing
+	AddLocalUser(LocalUserNum);
+
 	FullLoginCallback(LocalUserNum, AccountId, UserId);
 
 	GetLocalUserChecked(LocalUserNum).LastLoginCredentials_Legacy = MakeShared<FOnlineAccountCredentials>(AccountCredentials);
