Add lobby owner to member settings on creation, and changed a log level from warning to verbose

diff --git a/Engine/Plugins/Online/OnlineSubsystemEOS/Source/OnlineSubsystemEOS/Private/OnlineSessionEOS.cpp b/Engine/Plugins/Online/OnlineSubsystemEOS/Source/OnlineSubsystemEOS/Private/OnlineSessionEOS.cpp
--- a/Engine/Plugins/Online/OnlineSubsystemEOS/Source/OnlineSubsystemEOS/Private/OnlineSessionEOS.cpp
+++ b/Engine/Plugins/Online/OnlineSubsystemEOS/Source/OnlineSubsystemEOS/Private/OnlineSessionEOS.cpp
@@ -3566,7 +3566,7 @@ bool FOnlineSessionEOS::UnregisterPlayers(FName SessionName, const TArray< FUniq
 			}
 			else
 			{
-				UE_LOG_ONLINE_SESSION(Warning, TEXT("Player %s is not part of session (%s)"), *PlayerId->ToDebugString(), *SessionName.ToString());
+				UE_LOG_ONLINE_SESSION(Verbose, TEXT("Player %s is not a registered player of session (%s)"), *PlayerId->ToDebugString(), *SessionName.ToString());
 			}
 		}
 		if (bUnregisterEOS && EOSIds.Num() > 0)
@@ -4144,6 +4144,9 @@ uint32 FOnlineSessionEOS::CreateLobbySession(int32 HostingPlayerNum, FNamedOnlin
 			{
 				UE_LOG_ONLINE_SESSION(Verbose, TEXT("[FOnlineSessionEOS::CreateLobbySession] CreateLobby was successful. LobbyId is %s."), UTF8_TO_TCHAR(Data->LobbyId));
 
+				// Add the lobby owner to the member settings 
+				AddOnlineSessionMember(SessionName, LocalUserNetId.ToSharedRef()); 
+				
 				Session->SessionState = EOnlineSessionState::Pending;
 
 				TSharedPtr<FOnlineSessionInfoEOS> SessionInfo = MakeShared<FOnlineSessionInfoEOS>(FOnlineSessionInfoEOS::Create(FUniqueNetIdEOSLobby::Create(UTF8_TO_TCHAR(Data->LobbyId))));
