Update to AutoRefreshConnectLogin to support EOS as default OSS without EAS

diff --git a/Engine/Plugins/Online/OnlineSubsystemEOS/Source/OnlineSubsystemEOS/Private/UserManagerEOS.cpp b/Engine/Plugins/Online/OnlineSubsystemEOS/Source/OnlineSubsystemEOS/Private/UserManagerEOS.cpp
--- a/Engine/Plugins/Online/OnlineSubsystemEOS/Source/OnlineSubsystemEOS/Private/UserManagerEOS.cpp
+++ b/Engine/Plugins/Online/OnlineSubsystemEOS/Source/OnlineSubsystemEOS/Private/UserManagerEOS.cpp
@@ -1510,9 +1510,9 @@ void FUserManagerEOS::AutoRefreshConnectLogin(int32 LocalUserNum)
 	bool bIsDefaultOSS = DefaultOSS == TEXT("EOS");
 
 	const FEOSSettings Settings = UEOSSettings::GetSettings();
-	// In the case where bIsDefaultOSS is true, FUserManagerEOS::Login will default to using EOS_Auth_Login regardless of the value that bUseEAS is set to
-	// This behaviour will be fixed as part of a wider refactor of FUserManagerEOS::Login
-	const bool bShouldUseEOSAuthToken = bIsDefaultOSS || Settings.bUseEAS;
+
+	const bool bShouldUseEOSAuthToken = Settings.bUseNewLoginFlow ? Settings.bUseEAS : (bIsDefaultOSS || Settings.bUseEAS); 
+
 	if (bShouldUseEOSAuthToken)
 	{
 		const EOS_EpicAccountId AccountId = GetLocalEpicAccountId(LocalUserNum);
@@ -1605,12 +1605,11 @@ void FUserManagerEOS::AutoRefreshConnectLogin(int32 LocalUserNum)
 
 							// Now login into our EOS account
 
-							const FOnlineAccountCredentials& Creds = *GetLocalUserChecked(LocalUserNum).LastLoginCredentials_Legacy;
 							IOnlineSubsystem* PlatformOSS = GetPlatformOSS();
 							FEOSSettings Settings = UEOSSettings::GetSettings();
 							EOS_EExternalCredentialType CredType = Settings.bUseNewLoginFlow
 								? ToEOS_EExternalCredentialType(PlatformOSS ? PlatformOSS->GetSubsystemName() : EOSSubsystem->GetSubsystemName())
-								: ToEOS_EExternalCredentialTypeLegacy(PlatformOSS ? PlatformOSS->GetSubsystemName() : EOSSubsystem->GetSubsystemName(), Creds);
+								: ToEOS_EExternalCredentialTypeLegacy(PlatformOSS ? PlatformOSS->GetSubsystemName() : EOSSubsystem->GetSubsystemName(), *GetLocalUserChecked(LocalUserNum).LastLoginCredentials_Legacy);
 
 							FConnectCredentials Credentials(CredType, ConvertExternalAuthTokenToFString(AuthToken));
 
