Added validity check to prevent crash when the callback is called after the user has logged out.

diff --git a/Engine/Plugins/Online/OnlineSubsystemEOS/Source/OnlineSubsystemEOS/Private/UserManagerEOS.cpp b/Engine/Plugins/Online/OnlineSubsystemEOS/Source/OnlineSubsystemEOS/Private/UserManagerEOS.cpp
--- a/Engine/Plugins/Online/OnlineSubsystemEOS/Source/OnlineSubsystemEOS/Private/UserManagerEOS.cpp
+++ b/Engine/Plugins/Online/OnlineSubsystemEOS/Source/OnlineSubsystemEOS/Private/UserManagerEOS.cpp
@@ -4132,6 +4132,7 @@ bool FUserManagerEOS::QueryExternalIdMappings(const FUniqueNetId& UserId, const
 		FQueryByStringIdsCallback* CallbackObj = new FQueryByStringIdsCallback(AsWeak());
 		CallbackObj->CallbackLambda = [LocalUserNum, QueryOptions, BatchIds = MoveTemp(BatchIds), this, Delegate](const EOS_Connect_QueryExternalAccountMappingsCallbackInfo* Data)
 		{
+
 			EOS_EResult Result = Data->ResultCode;
 			if (GetLoginStatus(LocalUserNum) != ELoginStatus::LoggedIn)
 			{
@@ -4164,10 +4165,13 @@ bool FUserManagerEOS::QueryExternalIdMappings(const FUniqueNetId& UserId, const
 			}
 
 			// Mark all queries as complete
-			FLocalUserEOS& LocalUser = GetLocalUserChecked(LocalUserNum);
-			for (const FString& StringId : BatchIds)
+			if (IsLocalUserValid(LocalUserNum))
 			{
-				LocalUser.OngoingPlayerQueryExternalMappings.RemoveSwap(StringId, EAllowShrinking::No);
+				FLocalUserEOS& LocalUser = GetLocalUserChecked(LocalUserNum);
+				for (const FString& StringId : BatchIds)
+				{
+					LocalUser.OngoingPlayerQueryExternalMappings.RemoveSwap(StringId, EAllowShrinking::No);
+				}
 			}
 
 			const bool bWasSuccessful = Result == EOS_EResult::EOS_Success;
