Fix Connect Login refresh which was failing due to missing UserLoginInfo

diff --git a/Engine/Plugins/Online/OnlineSubsystemEOS/Source/OnlineSubsystemEOS/Private/UserManagerEOS.cpp b/Engine/Plugins/Online/OnlineSubsystemEOS/Source/OnlineSubsystemEOS/Private/UserManagerEOS.cpp
--- a/Engine/Plugins/Online/OnlineSubsystemEOS/Source/OnlineSubsystemEOS/Private/UserManagerEOS.cpp
+++ b/Engine/Plugins/Online/OnlineSubsystemEOS/Source/OnlineSubsystemEOS/Private/UserManagerEOS.cpp
@@ -1576,6 +1576,16 @@ void FUserManagerEOS::AutoRefreshConnectLogin(int32 LocalUserNum)
 				UE_EOS_CHECK_API_MISMATCH(EOS_CONNECT_LOGIN_API_LATEST, 2);
 				Options.Credentials = &Credentials;
 
+#if ADD_USER_LOGIN_INFO
+				EOS_Connect_UserLoginInfo UserLoginInfo = {};
+				UserLoginInfo.ApiVersion = 2;
+				UE_EOS_CHECK_API_MISMATCH(EOS_CONNECT_USERLOGININFO_API_LATEST, 2);
+				const FTCHARToUTF8 DisplayNameUtf8(*GetPlatformDisplayName(LocalUserNum));
+				UserLoginInfo.DisplayName = DisplayNameUtf8.Get();
+				UserLoginInfo.NsaIdToken = nullptr; // Not set when EAS is not used
+
+				Options.UserLoginInfo = &UserLoginInfo;
+#endif
 				FConnectLoginCallback* CallbackObj = new FConnectLoginCallback(AsWeak());
 				CallbackObj->CallbackLambda = [this, LocalUserNum](const EOS_Connect_LoginCallbackInfo* Data)
 					{
@@ -1618,6 +1628,16 @@ void FUserManagerEOS::AutoRefreshConnectLogin(int32 LocalUserNum)
 							UE_EOS_CHECK_API_MISMATCH(EOS_CONNECT_LOGIN_API_LATEST, 2);
 							Options.Credentials = &Credentials;
 
+#if ADD_USER_LOGIN_INFO
+							EOS_Connect_UserLoginInfo UserLoginInfo = {};
+							UserLoginInfo.ApiVersion = 2;
+							UE_EOS_CHECK_API_MISMATCH(EOS_CONNECT_USERLOGININFO_API_LATEST, 2);
+							const FTCHARToUTF8 DisplayNameUtf8(*GetPlatformDisplayName(LocalUserNum));
+							UserLoginInfo.DisplayName = DisplayNameUtf8.Get();
+							UserLoginInfo.NsaIdToken = nullptr; // Not set when EAS is not used
+
+							Options.UserLoginInfo = &UserLoginInfo;
+#endif
 							FConnectLoginCallback* CallbackObj = new FConnectLoginCallback(AsWeak());
 							CallbackObj->CallbackLambda = [this, LocalUserNum](const EOS_Connect_LoginCallbackInfo* Data)
 							{
