Changes to fix Beacon related disconnection issues, caused by the wrong elements being used as SocketId

diff --git a/Engine/Plugins/Online/OnlineSubsystemEOS/Source/OnlineSubsystemEOS/Private/OnlineSessionEOS.cpp b/Engine/Plugins/Online/OnlineSubsystemEOS/Source/OnlineSubsystemEOS/Private/OnlineSessionEOS.cpp
--- a/Engine/Plugins/Online/OnlineSubsystemEOS/Source/OnlineSubsystemEOS/Private/OnlineSessionEOS.cpp
+++ b/Engine/Plugins/Online/OnlineSubsystemEOS/Source/OnlineSubsystemEOS/Private/OnlineSessionEOS.cpp
@@ -3242,7 +3242,7 @@ bool FOnlineSessionEOS::PingSearchResults(const FOnlineSessionSearchResult& Sear
 }
 
 /** Get a resolved connection string from a session info */
-static bool GetConnectStringFromSessionInfo(TSharedPtr<FOnlineSessionInfoEOS>& SessionInfo, FString& ConnectInfo, int32 PortOverride=0)
+static bool GetConnectStringFromSessionInfo(TSharedPtr<FOnlineSessionInfoEOS>& SessionInfo, FString& ConnectInfo, FName SocketNameOverride = NAME_None, int32 PortOverride = 0)
 {
 	if (!SessionInfo.IsValid() || !SessionInfo->HostAddr.IsValid())
 	{
@@ -3258,6 +3258,11 @@ static bool GetConnectStringFromSessionInfo(TSharedPtr<FOnlineSessionInfoEOS>& S
 		ConnectInfo = SessionInfo->HostAddr->ToString(true);
 	}
 
+	if (SocketNameOverride != NAME_None)
+	{
+		ConnectInfo.ReplaceInline(*FName(NAME_GameNetDriver).ToString(), *SocketNameOverride.ToString());
+	}
+
 	return true;
 }
 
@@ -3272,7 +3277,7 @@ bool FOnlineSessionEOS::GetResolvedConnectString(FName SessionName, FString& Con
 		if (PortType == NAME_BeaconPort)
 		{
 			int32 BeaconListenPort = GetBeaconPortFromSessionSettings(Session->SessionSettings);
-			bSuccess = GetConnectStringFromSessionInfo(SessionInfo, ConnectInfo, BeaconListenPort);
+			bSuccess = GetConnectStringFromSessionInfo(SessionInfo, ConnectInfo, NAME_BeaconNetDriver, BeaconListenPort);
 		}
 		else if (PortType == NAME_GamePort)
 		{
@@ -3304,8 +3309,7 @@ bool FOnlineSessionEOS::GetResolvedConnectString(const FOnlineSessionSearchResul
 		if (PortType == NAME_BeaconPort)
 		{
 			int32 BeaconListenPort = GetBeaconPortFromSessionSettings(SearchResult.Session.SessionSettings);
-			bSuccess = GetConnectStringFromSessionInfo(SessionInfo, ConnectInfo, BeaconListenPort);
-
+			bSuccess = GetConnectStringFromSessionInfo(SessionInfo, ConnectInfo, NAME_BeaconNetDriver, BeaconListenPort);
 		}
 		else if (PortType == NAME_GamePort)
 		{
diff --git a/Engine/Plugins/Online/SocketSubsystemEOS/Source/SocketSubsystemEOS/Private/NetDriverEOSBase.cpp b/Engine/Plugins/Online/SocketSubsystemEOS/Source/SocketSubsystemEOS/Private/NetDriverEOSBase.cpp
--- a/Engine/Plugins/Online/SocketSubsystemEOS/Source/SocketSubsystemEOS/Private/NetDriverEOSBase.cpp
+++ b/Engine/Plugins/Online/SocketSubsystemEOS/Source/SocketSubsystemEOS/Private/NetDriverEOSBase.cpp
@@ -74,7 +74,7 @@ bool UNetDriverEOSBase::InitBase(bool bInitAsClient, FNetworkNotify* InNotify, c
 	TSharedRef<FInternetAddrEOS> EOSLocalAddress = StaticCastSharedRef<FInternetAddrEOS>(LocalAddress);
 	EOSLocalAddress->SetChannel((uint8)URL.Port);
 	// Set our net driver name so we don't accept connections across net driver types
-	EOSLocalAddress->SetSocketName(NetDriverName.ToString());
+	EOSLocalAddress->SetSocketName(GetNetDriverDefinition().ToString());
 
 	static_cast<FSocketEOS*>(GetSocket())->SetLocalAddress(*EOSLocalAddress);
 
