Fix for invalid user errors in lobby join operation.

diff --git a/Engine/Plugins/Online/OnlineSubsystemEOS/Source/OnlineSubsystemEOS/Private/OnlineSessionEOS.cpp b/Engine/Plugins/Online/OnlineSubsystemEOS/Source/OnlineSubsystemEOS/Private/OnlineSessionEOS.cpp
--- a/Engine/Plugins/Online/OnlineSubsystemEOS/Source/OnlineSubsystemEOS/Private/OnlineSessionEOS.cpp
+++ b/Engine/Plugins/Online/OnlineSubsystemEOS/Source/OnlineSubsystemEOS/Private/OnlineSessionEOS.cpp
@@ -873,7 +873,8 @@ void FOnlineSessionEOS::RegisterLobbyNotifications()
 void FOnlineSessionEOS::OnLobbyUpdateReceived(const EOS_LobbyId& LobbyId)
 {
 	const FUniqueNetIdEOSLobbyRef LobbyNetId = FUniqueNetIdEOSLobby::Create(UTF8_TO_TCHAR(LobbyId));
-	const EOS_ProductUserId LocalUserId = EOSSubsystem->UserManager->GetLocalProductUserId(GetDefaultLocalUserForLobby(*LobbyNetId));
+	// Because the update might happen before lobby members have been populated, we'll use the default local user here to ensure its validity
+	const EOS_ProductUserId LocalUserId = EOSSubsystem->UserManager->GetLocalProductUserId(EOSSubsystem->UserManager->GetDefaultLocalUser());
 
 	FNamedOnlineSession* Session = GetNamedSessionFromLobbyId(*LobbyNetId);
 	if (Session)
