Add a log message to EOS_Ecom_Checkout callback when EOS_Ecom_PurcahseProcessing is returned to let the caller know to call the OSS API that calls EOS_Ecom_QueryEntitlement

diff --git a/Engine/Plugins/Online/OnlineServicesEOS/Source/Private/Online/CommerceEOS.cpp b/Engine/Plugins/Online/OnlineServicesEOS/Source/Private/Online/CommerceEOS.cpp
--- a/Engine/Plugins/Online/OnlineServicesEOS/Source/Private/Online/CommerceEOS.cpp
+++ b/Engine/Plugins/Online/OnlineServicesEOS/Source/Private/Online/CommerceEOS.cpp
@@ -244,7 +244,19 @@ TOnlineAsyncOpHandle<FCommerceCheckout> FCommerceEOS::Checkout(FCommerceCheckout
 		EOS_EResult Result = Data->ResultCode;
 		if (Result != EOS_EResult::EOS_Success)
 		{
-			UE_LOG(LogOnlineServices, Error, TEXT("EOS_Ecom_Checkout: failed with error (%s)"), *LexToString(Data->ResultCode));
+			if (Result == EOS_EResult::EOS_Canceled)
+			{
+				UE_LOG(LogOnlineServices, Log, TEXT("EOS_Ecom_Checkout was canceled by user."));
+			}
+			else if (Result == EOS_EResult::EOS_Ecom_PurchaseProcessing)
+			{
+				// See header documentation for EOS_Ecom_PurchaseProcessing and online EOS documentation for more information
+				UE_LOG(LogOnlineServices, Log, TEXT("QueryEntitlements will need to be called to confirm the purchase. This can be done via a UI button."));
+			}
+			else
+			{
+				UE_LOG(LogOnlineServices, Error, TEXT("EOS_Ecom_Checkout: failed with error (%s)"), *LexToString(Data->ResultCode));
+			}
 			Op.SetError(Errors::FromEOSResult(Result));
 			return;
 		}
diff --git a/Engine/Plugins/Online/OnlineSubsystemEOS/Source/OnlineSubsystemEOS/Private/OnlineStoreEOS.cpp b/Engine/Plugins/Online/OnlineSubsystemEOS/Source/OnlineSubsystemEOS/Private/OnlineStoreEOS.cpp
--- a/Engine/Plugins/Online/OnlineSubsystemEOS/Source/OnlineSubsystemEOS/Private/OnlineStoreEOS.cpp
+++ b/Engine/Plugins/Online/OnlineSubsystemEOS/Source/OnlineSubsystemEOS/Private/OnlineStoreEOS.cpp
@@ -198,17 +198,20 @@ void FOnlineStoreEOS::Checkout(const FUniqueNetId& UserId, const FPurchaseChecko
 		const EOS_EResult CheckoutResult = Data->ResultCode;
 		if (CheckoutResult != EOS_EResult::EOS_Success)
 		{
-			UE_LOG_ONLINE(Error, TEXT("EOS_Ecom_Checkout: failed with error (%s)"), *LexToString(CheckoutResult));
 			if (CheckoutResult == EOS_EResult::EOS_Canceled)
 			{
+				UE_LOG_ONLINE(Log, TEXT("EOS_Ecom_Checkout was canceled by user."));
 				OnComplete.ExecuteIfBound(ONLINE_ERROR(EOnlineErrorResult::Canceled), MakeShared<FPurchaseReceipt>());
 			}
 			else if (CheckoutResult == EOS_EResult::EOS_Ecom_PurchaseProcessing)
 			{
+				// See header documentation for EOS_Ecom_PurchaseProcessing and online EOS documentation for more information
+				UE_LOG_ONLINE(Log, TEXT("QueryReceipts will need to be called to confirm the purchase. This can be done via a UI button."));
 				OnComplete.ExecuteIfBound(ONLINE_ERROR(EOnlineErrorResult::FailExtended, LexToString(CheckoutResult)), MakeShared<FPurchaseReceipt>());
 			}
 			else
 			{
+				UE_LOG_ONLINE(Error, TEXT("EOS_Ecom_Checkout: failed with error (%s)"), *LexToString(CheckoutResult));
 				OnComplete.ExecuteIfBound(ONLINE_ERROR(EOnlineErrorResult::Unknown), MakeShared<FPurchaseReceipt>());
 			}
 			return;
