Fix a crash in FUserManagerEOS that looked to happen when it failed to find UniqueNetId.  No repro, fix is based upon code inspection and crash reports.
FUserManagerEOS::ResolveUniqueNetId()

diff --git a/Engine/Plugins/Online/OnlineSubsystemEOS/Source/OnlineSubsystemEOS/Private/UserManagerEOS.cpp b/Engine/Plugins/Online/OnlineSubsystemEOS/Source/OnlineSubsystemEOS/Private/UserManagerEOS.cpp
--- a/Engine/Plugins/Online/OnlineSubsystemEOS/Source/OnlineSubsystemEOS/Private/UserManagerEOS.cpp
+++ b/Engine/Plugins/Online/OnlineSubsystemEOS/Source/OnlineSubsystemEOS/Private/UserManagerEOS.cpp
@@ -2337,8 +2337,16 @@ void FUserManagerEOS::ResolveUniqueNetIds(int32 LocalUserNum, const TArray<EOS_P
 
 void FUserManagerEOS::ResolveUniqueNetId(int32 LocalUserNum, const EOS_EpicAccountId& EpicAccountId, const FResolveUniqueNetIdCallback& Callback)
 {
-	const FResolveEpicAccountIdsCallback GroupCallback = [EpicAccountId, OriginalCallback = Callback](bool bWasSuccessful, TMap<EOS_EpicAccountId, FUniqueNetIdEOSRef> ResolvedUniqueNetIds, FString ErrorStr) {
-		OriginalCallback(ResolvedUniqueNetIds[EpicAccountId]);
+	const FResolveEpicAccountIdsCallback GroupCallback = [LocalUserNum, EpicAccountId, OriginalCallback = Callback](bool bWasSuccessful, TMap<EOS_EpicAccountId, FUniqueNetIdEOSRef> ResolvedUniqueNetIds, FString ErrorStr) {
+		if (bWasSuccessful)
+		{
+			OriginalCallback(ResolvedUniqueNetIds[EpicAccountId]);
+		}
+		else
+		{
+			UE_LOG_ONLINE(Error, TEXT("Unable to ResolveUniqueNetId for %d  %s"), LocalUserNum, *ErrorStr);
+			OriginalCallback(FUniqueNetIdEOS::EmptyId());
+		}
 	};
 
 	ResolveUniqueNetIds(LocalUserNum, { EpicAccountId }, GroupCallback);
