Remove local user if autologin fails in EOS OSS. Fix for Dedicated Server session creation flow.

diff --git a/Engine/Plugins/Online/OnlineSubsystemEOS/Source/OnlineSubsystemEOS/Private/UserManagerEOS.cpp b/Engine/Plugins/Online/OnlineSubsystemEOS/Source/OnlineSubsystemEOS/Private/UserManagerEOS.cpp
--- a/Engine/Plugins/Online/OnlineSubsystemEOS/Source/OnlineSubsystemEOS/Private/UserManagerEOS.cpp
+++ b/Engine/Plugins/Online/OnlineSubsystemEOS/Source/OnlineSubsystemEOS/Private/UserManagerEOS.cpp
@@ -1925,6 +1925,8 @@ bool FUserManagerEOS::AutoLogin(int32 LocalUserNum)
 		return true;
 	}
 
+	RemoveLocalUser(LocalUserNum);
+
 	UE_LOG_ONLINE(Warning, TEXT("Unable to AutoLogin user (%d). No valid configuration for AutoLogin."), LocalUserNum);
 	return false;
 }
