Added logic to FUserManagerEOS::AutoRefreshConnectLogin to set UserLoginInfo to the corresponding platform auth token when using EAS

diff --git a/Engine/Plugins/Online/OnlineSubsystemEOS/Source/OnlineSubsystemEOS/Private/UserManagerEOS.cpp b/Engine/Plugins/Online/OnlineSubsystemEOS/Source/OnlineSubsystemEOS/Private/UserManagerEOS.cpp
--- a/Engine/Plugins/Online/OnlineSubsystemEOS/Source/OnlineSubsystemEOS/Private/UserManagerEOS.cpp
+++ b/Engine/Plugins/Online/OnlineSubsystemEOS/Source/OnlineSubsystemEOS/Private/UserManagerEOS.cpp
@@ -1524,32 +1524,64 @@ void FUserManagerEOS::AutoRefreshConnectLogin(int32 LocalUserNum)
 		const FString AccessToken = GetAuthToken(LocalUserNum);
 		if (!AccessToken.IsEmpty())
 		{
-			// We update the auth token cached in the user account, along with the user information
-			const FUserOnlineAccountEOSRef UserAccountRef = GetLocalUserChecked(LocalUserNum).UserOnlineAccount.ToSharedRef();
-			UpdateUserInfo(UserAccountRef, AccountId, AccountId);
-
-			EOS_Connect_Credentials Credentials = { };
-			Credentials.ApiVersion = 1;
-			UE_EOS_CHECK_API_MISMATCH(EOS_CONNECT_CREDENTIALS_API_LATEST, 1);
-			Credentials.Type = EOS_EExternalCredentialType::EOS_ECT_EPIC;
-			auto AccessTokenConverter = StringCast<UTF8CHAR>(*AccessToken);
-			Credentials.Token = (const char*)AccessTokenConverter.Get();
-
-			EOS_Connect_LoginOptions Options = { };
-			Options.ApiVersion = 2;
-			UE_EOS_CHECK_API_MISMATCH(EOS_CONNECT_LOGIN_API_LATEST, 2);
-			Options.Credentials = &Credentials;
-
-			FConnectLoginCallback* CallbackObj = new FConnectLoginCallback(AsWeak());
-			CallbackObj->CallbackLambda = [LocalUserNum, AccountId, this](const EOS_Connect_LoginCallbackInfo* Data)
-			{
-				if (Data->ResultCode != EOS_EResult::EOS_Success)
+#if ADD_USER_LOGIN_INFO
+			GetPlatformAuthToken(LocalUserNum,
+				FOnGetLinkedAccountAuthTokenCompleteDelegate::CreateLambda([this, WeakThis = AsWeak(), AccessToken, AccountId](int32 LocalUserNum, bool bWasSuccessful, const FExternalAuthToken& AuthToken)
 				{
-					UE_LOG_ONLINE(Error, TEXT("Failed to refresh ConnectLogin(%d) failed with EOS result code (%s)"), LocalUserNum, *LexToString(Data->ResultCode));
-					Logout(LocalUserNum);
-				}
-			};
-			EOS_Connect_Login(EOSSubsystem->ConnectHandle, &Options, CallbackObj, CallbackObj->GetCallbackPtr());
+					if (FUserManagerEOSPtr StrongThis = WeakThis.Pin())
+					{
+						if (!bWasSuccessful || !AuthToken.IsValid())
+						{
+							const FString ErrorString = FString::Printf(TEXT("ConnectLoginEAS(%d) failed due to the platform OSS giving an empty auth token"), LocalUserNum);
+							UE_LOG_ONLINE(Warning, TEXT("%s"), *ErrorString);
+							RemoveLocalUser(LocalUserNum);
+							TriggerOnLoginCompleteDelegates(LocalUserNum, false, *FUniqueNetIdEOS::EmptyId(), ErrorString);
+							return;
+						}
+#endif
+						// We update the auth token cached in the user account, along with the user information
+						const FUserOnlineAccountEOSRef UserAccountRef = GetLocalUserChecked(LocalUserNum).UserOnlineAccount.ToSharedRef();
+						UpdateUserInfo(UserAccountRef, AccountId, AccountId);
+
+						EOS_Connect_Credentials Credentials = { };
+						Credentials.ApiVersion = 1;
+						UE_EOS_CHECK_API_MISMATCH(EOS_CONNECT_CREDENTIALS_API_LATEST, 1);
+						Credentials.Type = EOS_EExternalCredentialType::EOS_ECT_EPIC;
+						auto AccessTokenConverter = StringCast<UTF8CHAR>(*AccessToken);
+						Credentials.Token = (const char*)AccessTokenConverter.Get();
+
+						EOS_Connect_LoginOptions Options = { };
+						Options.ApiVersion = 2;
+						UE_EOS_CHECK_API_MISMATCH(EOS_CONNECT_LOGIN_API_LATEST, 2);
+						Options.Credentials = &Credentials;
+
+#if ADD_USER_LOGIN_INFO
+						EOS_Connect_UserLoginInfo UserLoginInfo = {};
+						UserLoginInfo.ApiVersion = 2;
+						UE_EOS_CHECK_API_MISMATCH(EOS_CONNECT_USERLOGININFO_API_LATEST, 2);
+						const FTCHARToUTF8 DisplayNameUtf8(*GetPlatformDisplayName(LocalUserNum));
+						UserLoginInfo.DisplayName = DisplayNameUtf8.Get();
+						auto AuthTokenConverter = StringCast<UTF8CHAR>(*AuthToken.TokenString);
+						UserLoginInfo.NsaIdToken = (const char*)AuthTokenConverter.Get();
+
+						Options.UserLoginInfo = &UserLoginInfo;
+#endif
+
+						FConnectLoginCallback* CallbackObj = new FConnectLoginCallback(AsWeak());
+						CallbackObj->CallbackLambda = [LocalUserNum, AccountId, this](const EOS_Connect_LoginCallbackInfo* Data)
+							{
+								if (Data->ResultCode != EOS_EResult::EOS_Success)
+								{
+									UE_LOG_ONLINE(Error, TEXT("Failed to refresh ConnectLogin(%d) failed with EOS result code (%s)"), LocalUserNum, *LexToString(Data->ResultCode));
+									Logout(LocalUserNum);
+								}
+							};
+						EOS_Connect_Login(EOSSubsystem->ConnectHandle, &Options, CallbackObj, CallbackObj->GetCallbackPtr());
+
+#if ADD_USER_LOGIN_INFO
+					}
+				}));
+#endif
 		}
 		else
 		{
