Added logic to support EOS_Platform_Switch_SetNintendoProductUserIdMapping and callsite as part of OSS EOS Login

diff --git a/Engine/Plugins/Online/OnlineSubsystemEOS/Source/OnlineSubsystemEOS/Private/EOSHelpers.cpp b/Engine/Plugins/Online/OnlineSubsystemEOS/Source/OnlineSubsystemEOS/Private/EOSHelpers.cpp
--- a/Engine/Plugins/Online/OnlineSubsystemEOS/Source/OnlineSubsystemEOS/Private/EOSHelpers.cpp
+++ b/Engine/Plugins/Online/OnlineSubsystemEOS/Source/OnlineSubsystemEOS/Private/EOSHelpers.cpp
@@ -38,6 +38,10 @@ IEOSPlatformHandlePtr FEOSHelpers::CreatePlatform(EOS_Platform_Options& Platform
 	return nullptr;
 }
 
+void FEOSHelpers::AddProductUserIdMapping(IEOSPlatformHandlePtr PlatformHandlePtr, const EOS_ProductUserId& ProductUserId, int32 LocalUserNum)
+{
+}
+
 void FEOSHelpers::AddExternalAccountMapping(IEOSPlatformHandlePtr PlatformHandlePtr, const EOS_EpicAccountId& EpicAccountId, int32 LocalUserNum)
 {
 }
diff --git a/Engine/Plugins/Online/OnlineSubsystemEOS/Source/OnlineSubsystemEOS/Private/EOSHelpers.h b/Engine/Plugins/Online/OnlineSubsystemEOS/Source/OnlineSubsystemEOS/Private/EOSHelpers.h
--- a/Engine/Plugins/Online/OnlineSubsystemEOS/Source/OnlineSubsystemEOS/Private/EOSHelpers.h
+++ b/Engine/Plugins/Online/OnlineSubsystemEOS/Source/OnlineSubsystemEOS/Private/EOSHelpers.h
@@ -23,6 +23,7 @@ public:
 	virtual void GetSystemAuthCredentialsOptions(void*& OutCredentialsOptions);
 	virtual void PlatformTriggerLoginUI(FOnlineSubsystemEOS* EOSSubsystem, const int ControllerIndex, bool bShowOnlineOnly, bool bShowSkipButton, const FOnLoginUIClosedDelegate& Delegate);
 	virtual IEOSPlatformHandlePtr CreatePlatform(EOS_Platform_Options& PlatformOptions);
+	virtual void AddProductUserIdMapping(IEOSPlatformHandlePtr PlatformHandlePtr, const EOS_ProductUserId& ProductUserId, int32 LocalUserNum);
 	virtual void AddExternalAccountMapping(IEOSPlatformHandlePtr PlatformHandlePtr, const EOS_EpicAccountId& EpicAccountId, int32 LocalUserNum);
 
 protected:
diff --git a/Engine/Plugins/Online/OnlineSubsystemEOS/Source/OnlineSubsystemEOS/Private/UserManagerEOS.cpp b/Engine/Plugins/Online/OnlineSubsystemEOS/Source/OnlineSubsystemEOS/Private/UserManagerEOS.cpp
--- a/Engine/Plugins/Online/OnlineSubsystemEOS/Source/OnlineSubsystemEOS/Private/UserManagerEOS.cpp
+++ b/Engine/Plugins/Online/OnlineSubsystemEOS/Source/OnlineSubsystemEOS/Private/UserManagerEOS.cpp
@@ -1254,7 +1254,7 @@ void FUserManagerEOS::CallEOSConnectLogin(int32 LocalUserNum, const FOnlineAccou
 	UE_EOS_CHECK_API_MISMATCH(EOS_CONNECT_USERLOGININFO_API_LATEST, 2);
 	const FTCHARToUTF8 DisplayNameUtf8(*GetPlatformDisplayName(LocalUserNum));
 	UserLoginInfo.DisplayName = DisplayNameUtf8.Get();
-	UserLoginInfo.NsaIdToken = nullptr; // Not set when EAS is not used
+	UserLoginInfo.NsaIdToken = nullptr; // Must be null when EOS_ECT_NINTENDO_NSA_ID_TOKEN is used
 
 	Options.UserLoginInfo = &UserLoginInfo;
 #endif
@@ -1806,6 +1806,9 @@ void FUserManagerEOS::FullLoginCallback(int32 LocalUserNum, EOS_EpicAccountId Ac
 		PresenceNotificationId = EOS_Presence_AddNotifyOnPresenceChanged(EOSSubsystem->PresenceHandle, &Options, CallbackObj, CallbackObj->GetCallbackPtr());
 	}
 
+	// Set the ProductUserId mapping if the platform supports it
+	EOSSubsystem->GetEOSHelpers()->AddProductUserIdMapping(EOSSubsystem->GetEOSPlatformHandle(), UserId, LocalUserNum);
+
 	// We add the local user
 	FLocalUserEOS& LocalUser = UpdateLocalUser(LocalUserNum, AccountId, UserId);
 
