Add call to AddLocalUsers for non-EAS legacy login flow.

diff --git a/Engine/Plugins/Online/OnlineSubsystemEOS/Source/OnlineSubsystemEOS/Private/UserManagerEOS.cpp b/Engine/Plugins/Online/OnlineSubsystemEOS/Source/OnlineSubsystemEOS/Private/UserManagerEOS.cpp
--- a/Engine/Plugins/Online/OnlineSubsystemEOS/Source/OnlineSubsystemEOS/Private/UserManagerEOS.cpp
+++ b/Engine/Plugins/Online/OnlineSubsystemEOS/Source/OnlineSubsystemEOS/Private/UserManagerEOS.cpp
@@ -1260,17 +1260,32 @@ void FUserManagerEOS::CallEOSConnectLogin(int32 LocalUserNum, const FOnlineAccou
 #endif
 
 	FConnectLoginCallback* CallbackObj = new FConnectLoginCallback(AsWeak());
-	CallbackObj->CallbackLambda = [this, LocalUserNum](const EOS_Connect_LoginCallbackInfo* Data)
+	CallbackObj->CallbackLambda = [this, LocalUserNum, Credentials](const EOS_Connect_LoginCallbackInfo* Data)
 		{
+			FEOSSettings Settings = UEOSSettings::GetSettings();
 			if (Data->ResultCode == EOS_EResult::EOS_Success)
 			{
 				// We have an account mapping to the platform account, skip to final login
-				FullLoginCallback(LocalUserNum, nullptr, Data->LocalUserId);
+				if (Settings.bUseNewLoginFlow)
+				{
+					FullLoginCallback(LocalUserNum, nullptr, Data->LocalUserId);
+				}
+				else
+				{
+					FullLoginCallbackLegacy(LocalUserNum, nullptr, Data->LocalUserId, Credentials);
+				}
 			}
 			else if (Data->ResultCode == EOS_EResult::EOS_InvalidUser)
 			{
 				// We need to create the platform account mapping for this user using the continuation token
-				CreateConnectedLogin(LocalUserNum, nullptr, Data->ContinuanceToken);
+				if (Settings.bUseNewLoginFlow)
+				{
+					CreateConnectedLogin(LocalUserNum, nullptr, Data->ContinuanceToken);
+				}
+				else
+				{
+					CreateConnectedLoginLegacy(LocalUserNum, nullptr, Data->ContinuanceToken, Credentials);
+				}
 			}
 			else
 			{
