- Fix comments on bUseEAS and bUseEOSConnect. They were previously tagged as EOSPlus related, but they are used to control EOS login flow.

diff --git a/Engine/Plugins/Online/OnlineSubsystemEOS/Source/OnlineSubsystemEOS/Public/EOSSettings.h b/Engine/Plugins/Online/OnlineSubsystemEOS/Source/OnlineSubsystemEOS/Public/EOSSettings.h
--- a/Engine/Plugins/Online/OnlineSubsystemEOS/Source/OnlineSubsystemEOS/Public/EOSSettings.h
+++ b/Engine/Plugins/Online/OnlineSubsystemEOS/Source/OnlineSubsystemEOS/Public/EOSSettings.h
@@ -159,12 +159,12 @@ public:
 	UPROPERTY(Config, EditAnywhere, BlueprintReadOnly, Category = "EOS Settings")
 	TArray<FString> AuthScopeFlags;
 
-	/** Set to true to have Epic Accounts used (friends list will be unified with the default platform) */
-	UPROPERTY(Config, EditAnywhere, BlueprintReadOnly, Category="EOSPlus Login Settings", DisplayName="Use Epic Account for EOS login (requires account linking)")
+	/** Set to true to login to EOS_Auth (required to use Epic Account Services) */
+	UPROPERTY(Config, EditAnywhere, BlueprintReadOnly, Category="EOS Settings",  DisplayName="Login to EOS Auth, and use Epic Account Services")
 	bool bUseEAS = false;
 
-	/** Set to true to have EOS Connect APIs used to link accounts for crossplay */
-	UPROPERTY(Config, EditAnywhere, BlueprintReadOnly, Category="EOSPlus Login Settings", DisplayName="Use EOS Connect APIs to create and link Product User IDs (PUIDs), and use EOS Game Services")
+	/** Set to true to login to EOS_Connect (required to use Epic Game Services) */
+	UPROPERTY(Config, EditAnywhere, BlueprintReadOnly, Category="EOS Settings", DisplayName="Login to EOS Connect, and use EOS Game Services")
 	bool bUseEOSConnect = false;
 
 	/** Set to true to write stats to EOS as well as the default platform */
