Added functionality to public interface to convert EOS id type into UE id type

diff --git a/Engine/Plugins/Online/OnlineSubsystemEOS/Source/OnlineSubsystemEOS/Private/OnlineSessionEOS.cpp b/Engine/Plugins/Online/OnlineSubsystemEOS/Source/OnlineSubsystemEOS/Private/OnlineSessionEOS.cpp
--- a/Engine/Plugins/Online/OnlineSubsystemEOS/Source/OnlineSubsystemEOS/Private/OnlineSessionEOS.cpp
+++ b/Engine/Plugins/Online/OnlineSubsystemEOS/Source/OnlineSubsystemEOS/Private/OnlineSessionEOS.cpp
@@ -646,7 +646,7 @@ void FOnlineSessionEOS::RegisterSessionNotifications()
 
 void FOnlineSessionEOS::OnSessionInviteReceived(const EOS_Sessions_SessionInviteReceivedCallbackInfo* Data)
 {
-	EOSSubsystem->UserManager->ResolveUniqueNetIds(EOSSubsystem->UserManager->GetDefaultLocalUser(), { Data->LocalUserId, Data->TargetUserId }, [this, LocalUserId = Data->LocalUserId, TargetUserId = Data->TargetUserId, InviteId = FString(UTF8_TO_TCHAR(Data->InviteId))](TMap<EOS_ProductUserId, FUniqueNetIdEOSRef> ResolvedUniqueNetIds)
+	EOSSubsystem->UserManager->ResolveUniqueNetIds(EOSSubsystem->UserManager->GetDefaultLocalUser(), { Data->LocalUserId, Data->TargetUserId }, [this, LocalUserId = Data->LocalUserId, TargetUserId = Data->TargetUserId, InviteId = FString(UTF8_TO_TCHAR(Data->InviteId))](TMap<EOS_ProductUserId, FUniqueNetIdEOSRef> ResolvedUniqueNetIds, const FOnlineError& Error)
 		{
 			if (!ResolvedUniqueNetIds.Contains(LocalUserId))
 			{
@@ -692,7 +692,7 @@ void FOnlineSessionEOS::OnSessionInviteReceived(const EOS_Sessions_SessionInvite
 
 void FOnlineSessionEOS::OnSessionInviteAccepted(const EOS_Sessions_SessionInviteAcceptedCallbackInfo* Data)
 {
-	EOSSubsystem->UserManager->ResolveUniqueNetId(EOSSubsystem->UserManager->GetDefaultLocalUser(), Data->LocalUserId, [this, InviteId = FString(UTF8_TO_TCHAR(Data->InviteId))](FUniqueNetIdEOSRef ResolvedUniqueNetId)
+	EOSSubsystem->UserManager->ResolveUniqueNetId(EOSSubsystem->UserManager->GetDefaultLocalUser(), Data->LocalUserId, [this, InviteId = FString(UTF8_TO_TCHAR(Data->InviteId))](FUniqueNetIdEOSRef ResolvedUniqueNetId, const FOnlineError& Error)
 		{
 			const int32 LocalUserNum = EOSSubsystem->UserManager->GetLocalUserNumFromUniqueNetId(*ResolvedUniqueNetId);
 
@@ -722,7 +722,7 @@ void FOnlineSessionEOS::OnSessionInviteAccepted(const EOS_Sessions_SessionInvite
 
 void FOnlineSessionEOS::OnJoinSessionAccepted(const EOS_Sessions_JoinSessionAcceptedCallbackInfo* Data)
 {
-	EOSSubsystem->UserManager->ResolveUniqueNetId(EOSSubsystem->UserManager->GetDefaultLocalUser(), Data->LocalUserId, [this, UiEventId = Data->UiEventId](FUniqueNetIdEOSRef ResolvedUniqueNetId)
+	EOSSubsystem->UserManager->ResolveUniqueNetId(EOSSubsystem->UserManager->GetDefaultLocalUser(), Data->LocalUserId, [this, UiEventId = Data->UiEventId](FUniqueNetIdEOSRef ResolvedUniqueNetId, const FOnlineError& Error)
 		{
 			const int32 LocalUserNum = EOSSubsystem->UserManager->GetLocalUserNumFromUniqueNetId(*ResolvedUniqueNetId);
 
@@ -931,7 +931,7 @@ void FOnlineSessionEOS::OnLobbyMemberUpdateReceived(const EOS_LobbyId& LobbyId,
 {
 	const FUniqueNetIdEOSLobbyRef LobbyNetId = FUniqueNetIdEOSLobby::Create(UTF8_TO_TCHAR(LobbyId));
 
-	EOSSubsystem->UserManager->ResolveUniqueNetId(EOSSubsystem->UserManager->GetDefaultLocalUser(), TargetUserId, [this, TargetUserId, LobbyNetId](FUniqueNetIdEOSRef ResolvedUniqueNetId)
+	EOSSubsystem->UserManager->ResolveUniqueNetId(EOSSubsystem->UserManager->GetDefaultLocalUser(), TargetUserId, [this, TargetUserId, LobbyNetId](FUniqueNetIdEOSRef ResolvedUniqueNetId, const FOnlineError& Error)
 		{
 			UpdateOrAddLobbyMember(LobbyNetId, ResolvedUniqueNetId);
 		});
@@ -947,7 +947,7 @@ void FOnlineSessionEOS::OnMemberStatusReceived(const EOS_LobbyId& LobbyId, const
 		{
 		case EOS_ELobbyMemberStatus::EOS_LMS_JOINED:
 			{
-				EOSSubsystem->UserManager->ResolveUniqueNetId(EOSSubsystem->UserManager->GetDefaultLocalUser(), TargetUserId, [this, LobbyNetId](FUniqueNetIdEOSRef ResolvedUniqueNetId)
+				EOSSubsystem->UserManager->ResolveUniqueNetId(EOSSubsystem->UserManager->GetDefaultLocalUser(), TargetUserId, [this, LobbyNetId](FUniqueNetIdEOSRef ResolvedUniqueNetId, const FOnlineError& Error)
 					{
 						UpdateOrAddLobbyMember(LobbyNetId, ResolvedUniqueNetId);
 					});
@@ -956,7 +956,7 @@ void FOnlineSessionEOS::OnMemberStatusReceived(const EOS_LobbyId& LobbyId, const
 			break;
 		case EOS_ELobbyMemberStatus::EOS_LMS_LEFT:
 			{
-				EOSSubsystem->UserManager->ResolveUniqueNetId(EOSSubsystem->UserManager->GetDefaultLocalUser(), TargetUserId, [this, LobbyNetId](FUniqueNetIdEOSRef ResolvedUniqueNetId)
+				EOSSubsystem->UserManager->ResolveUniqueNetId(EOSSubsystem->UserManager->GetDefaultLocalUser(), TargetUserId, [this, LobbyNetId](FUniqueNetIdEOSRef ResolvedUniqueNetId, const FOnlineError& Error)
 					{
 						FNamedOnlineSession* Session = GetNamedSessionFromLobbyId(*LobbyNetId);
 						if (Session)
@@ -974,7 +974,7 @@ void FOnlineSessionEOS::OnMemberStatusReceived(const EOS_LobbyId& LobbyId, const
 			break;
 		case EOS_ELobbyMemberStatus::EOS_LMS_DISCONNECTED:
 			{
-				EOSSubsystem->UserManager->ResolveUniqueNetId(EOSSubsystem->UserManager->GetDefaultLocalUser(), TargetUserId, [this, LobbyNetId](FUniqueNetIdEOSRef ResolvedUniqueNetId)
+				EOSSubsystem->UserManager->ResolveUniqueNetId(EOSSubsystem->UserManager->GetDefaultLocalUser(), TargetUserId, [this, LobbyNetId](FUniqueNetIdEOSRef ResolvedUniqueNetId, const FOnlineError& Error)
 					{
 						FNamedOnlineSession* Session = GetNamedSessionFromLobbyId(*LobbyNetId);
 						if (Session)
@@ -992,7 +992,7 @@ void FOnlineSessionEOS::OnMemberStatusReceived(const EOS_LobbyId& LobbyId, const
 			break;
 		case EOS_ELobbyMemberStatus::EOS_LMS_KICKED:
 			{
-				EOSSubsystem->UserManager->ResolveUniqueNetId(EOSSubsystem->UserManager->GetDefaultLocalUser(), TargetUserId, [this, LobbyNetId](FUniqueNetIdEOSRef ResolvedUniqueNetId)
+				EOSSubsystem->UserManager->ResolveUniqueNetId(EOSSubsystem->UserManager->GetDefaultLocalUser(), TargetUserId, [this, LobbyNetId](FUniqueNetIdEOSRef ResolvedUniqueNetId, const FOnlineError& Error)
 					{
 						FNamedOnlineSession* Session = GetNamedSessionFromLobbyId(*LobbyNetId);
 						if (Session)
@@ -1010,7 +1010,7 @@ void FOnlineSessionEOS::OnMemberStatusReceived(const EOS_LobbyId& LobbyId, const
 			break;
 		case EOS_ELobbyMemberStatus::EOS_LMS_PROMOTED:
 			{
-				EOSSubsystem->UserManager->ResolveUniqueNetId(EOSSubsystem->UserManager->GetDefaultLocalUser(), TargetUserId, [this, LobbyNetId](FUniqueNetIdEOSRef ResolvedUniqueNetId)
+				EOSSubsystem->UserManager->ResolveUniqueNetId(EOSSubsystem->UserManager->GetDefaultLocalUser(), TargetUserId, [this, LobbyNetId](FUniqueNetIdEOSRef ResolvedUniqueNetId, const FOnlineError& Error)
 					{
 						FNamedOnlineSession* Session = GetNamedSessionFromLobbyId(*LobbyNetId);
 						if (Session)
@@ -1057,7 +1057,7 @@ void FOnlineSessionEOS::OnMemberStatusReceived(const EOS_LobbyId& LobbyId, const
 void FOnlineSessionEOS::OnLobbyInviteReceived(const EOS_Lobby_LobbyInviteReceivedCallbackInfo* Data)
 {
 	const TArray<EOS_ProductUserId> PUIdS = { Data->LocalUserId, Data->TargetUserId };
-	EOSSubsystem->UserManager->ResolveUniqueNetIds(EOSSubsystem->UserManager->GetDefaultLocalUser(), PUIdS, [this, LocalUserId = Data->LocalUserId, TargetUserId = Data->TargetUserId, InviteId = FString(UTF8_TO_TCHAR(Data->InviteId))](TMap<EOS_ProductUserId, FUniqueNetIdEOSRef> ResolvedUniqueNetIds)
+	EOSSubsystem->UserManager->ResolveUniqueNetIds(EOSSubsystem->UserManager->GetDefaultLocalUser(), PUIdS, [this, LocalUserId = Data->LocalUserId, TargetUserId = Data->TargetUserId, InviteId = FString(UTF8_TO_TCHAR(Data->InviteId))](TMap<EOS_ProductUserId, FUniqueNetIdEOSRef> ResolvedUniqueNetIds, const FOnlineError& Error)
 		{
 			if (!ResolvedUniqueNetIds.Contains(LocalUserId))
 			{
@@ -1104,7 +1104,7 @@ void FOnlineSessionEOS::OnLobbyInviteReceived(const EOS_Lobby_LobbyInviteReceive
 
 void FOnlineSessionEOS::OnLobbyInviteAccepted(const EOS_Lobby_LobbyInviteAcceptedCallbackInfo* Data)
 {
-	EOSSubsystem->UserManager->ResolveUniqueNetId(EOSSubsystem->UserManager->GetDefaultLocalUser(), Data->LocalUserId, [this, InviteId = FString(UTF8_TO_TCHAR(Data->InviteId))](FUniqueNetIdEOSRef ResolvedUniqueNetId)
+	EOSSubsystem->UserManager->ResolveUniqueNetId(EOSSubsystem->UserManager->GetDefaultLocalUser(), Data->LocalUserId, [this, InviteId = FString(UTF8_TO_TCHAR(Data->InviteId))](FUniqueNetIdEOSRef ResolvedUniqueNetId, const FOnlineError& Error)
 		{
 			const int32 LocalUserNum = EOSSubsystem->UserManager->GetLocalUserNumFromUniqueNetId(*ResolvedUniqueNetId);
 
@@ -1138,7 +1138,7 @@ void FOnlineSessionEOS::OnLobbyInviteAccepted(const EOS_Lobby_LobbyInviteAccepte
 
 void FOnlineSessionEOS::OnJoinLobbyAccepted(const EOS_Lobby_JoinLobbyAcceptedCallbackInfo* Data)
 {
-	EOSSubsystem->UserManager->ResolveUniqueNetId(EOSSubsystem->UserManager->GetDefaultLocalUser(), Data->LocalUserId, [this, UiEventId = Data->UiEventId](FUniqueNetIdEOSRef ResolvedUniqueNetId)
+	EOSSubsystem->UserManager->ResolveUniqueNetId(EOSSubsystem->UserManager->GetDefaultLocalUser(), Data->LocalUserId, [this, UiEventId = Data->UiEventId](FUniqueNetIdEOSRef ResolvedUniqueNetId, const FOnlineError& Error)
 		{
 			const int32 LocalUserNum = EOSSubsystem->UserManager->GetLocalUserNumFromUniqueNetId(*ResolvedUniqueNetId);
 
@@ -2424,7 +2424,7 @@ void FOnlineSessionEOS::CopySearchResult(const FSessionDetailsEOS& SessionHandle
 	}
 	else
 	{
-		EOSSubsystem->UserManager->ResolveUniqueNetId(EOSSubsystem->UserManager->GetDefaultLocalUser(), SessionInfo->OwnerUserId, [this, &OutSession, Callback = MoveTemp(Callback)](FUniqueNetIdEOSRef ResolvedUniqueNetId)
+		EOSSubsystem->UserManager->ResolveUniqueNetId(EOSSubsystem->UserManager->GetDefaultLocalUser(), SessionInfo->OwnerUserId, [this, &OutSession, Callback = MoveTemp(Callback)](FUniqueNetIdEOSRef ResolvedUniqueNetId, const FOnlineError& Error)
 			{
 				OutSession.OwningUserId = ResolvedUniqueNetId;
 				OutSession.OwningUserName = EOSSubsystem->UserManager->GetPlayerNickname(*ResolvedUniqueNetId);
@@ -4931,7 +4931,7 @@ void FOnlineSessionEOS::CopyLobbyData(const TSharedRef<FLobbyDetailsEOS>& LobbyD
 
 		if (!TargetUserIds.IsEmpty())
 		{
-			EOSSubsystem->UserManager->ResolveUniqueNetIds(EOSSubsystem->UserManager->GetDefaultLocalUser(), TargetUserIds, [this, LobbyDetails, LobbyId = FUniqueNetIdEOSLobby::Create(UTF8_TO_TCHAR(LobbyDetailsInfo->LobbyId)), Callback = MoveTemp(Callback)](TMap<EOS_ProductUserId, FUniqueNetIdEOSRef> ResolvedUniqueNetIds)
+			EOSSubsystem->UserManager->ResolveUniqueNetIds(EOSSubsystem->UserManager->GetDefaultLocalUser(), TargetUserIds, [this, LobbyDetails, LobbyId = FUniqueNetIdEOSLobby::Create(UTF8_TO_TCHAR(LobbyDetailsInfo->LobbyId)), Callback = MoveTemp(Callback)](TMap<EOS_ProductUserId, FUniqueNetIdEOSRef> ResolvedUniqueNetIds, const FOnlineError& Error)
 				{
 					FOnlineSession* Session = GetOnlineSessionFromLobbyId(*LobbyId);
 					if (Session)
@@ -4978,7 +4978,7 @@ void FOnlineSessionEOS::CopyLobbyData(const TSharedRef<FLobbyDetailsEOS>& LobbyD
 
 		const EOS_ProductUserId LobbyOwner = EOS_LobbyDetails_GetLobbyOwner(LobbyDetails->LobbyDetailsHandle, &GetLobbyOwnerOptions);
 
-		EOSSubsystem->UserManager->ResolveUniqueNetIds(EOSSubsystem->UserManager->GetDefaultLocalUser(), { LobbyOwner }, [this, LobbyOwner, LobbyDetails, LobbyId = FUniqueNetIdEOSLobby::Create(UTF8_TO_TCHAR(LobbyDetailsInfo->LobbyId)), Callback = MoveTemp(Callback)](TMap<EOS_ProductUserId, FUniqueNetIdEOSRef> ResolvedUniqueNetIds)
+		EOSSubsystem->UserManager->ResolveUniqueNetIds(EOSSubsystem->UserManager->GetDefaultLocalUser(), { LobbyOwner }, [this, LobbyOwner, LobbyDetails, LobbyId = FUniqueNetIdEOSLobby::Create(UTF8_TO_TCHAR(LobbyDetailsInfo->LobbyId)), Callback = MoveTemp(Callback)](TMap<EOS_ProductUserId, FUniqueNetIdEOSRef> ResolvedUniqueNetIds, const FOnlineError& Error)
 			{
 				FOnlineSession* Session = GetOnlineSessionFromLobbyId(*LobbyId);
 				if (Session)
diff --git a/Engine/Plugins/Online/OnlineSubsystemEOS/Source/OnlineSubsystemEOS/Private/OnlineSubsystemEOS.cpp b/Engine/Plugins/Online/OnlineSubsystemEOS/Source/OnlineSubsystemEOS/Private/OnlineSubsystemEOS.cpp
--- a/Engine/Plugins/Online/OnlineSubsystemEOS/Source/OnlineSubsystemEOS/Private/OnlineSubsystemEOS.cpp
+++ b/Engine/Plugins/Online/OnlineSubsystemEOS/Source/OnlineSubsystemEOS/Private/OnlineSubsystemEOS.cpp
@@ -729,6 +729,28 @@ IVoiceChatUser* FOnlineSubsystemEOS::GetVoiceChatUserInterface(const FUniqueNetI
 	return Result;
 }
 
+void FOnlineSubsystemEOS::QueryUniqueNetId(int32 LocalUserNum, const EOS_ProductUserId& ProductUserId, const FOnQueryUniqueNetIdComplete & Callback)
+{
+	if (UserManager)
+	{
+		FUserManagerEOS::FResolveUniqueNetIdCallback ResolveUniqueNetIdCallback = [Callback](FUniqueNetIdEOSRef ResolvedUniqueNetId, const FOnlineError& Error)
+		{
+			Callback.ExecuteIfBound(ResolvedUniqueNetId, Error);
+		};
+
+		UserManager->ResolveUniqueNetId(LocalUserNum, ProductUserId, ResolveUniqueNetIdCallback);
+	}
+	else
+	{
+		UE_LOG_ONLINE(Warning, TEXT("Unable to query UniqueNetId for ProductUserId. UserManager not valid."));
+
+		ExecuteNextTick([Callback]()
+		{
+			Callback.ExecuteIfBound(FUniqueNetIdEOS::EmptyId(), FOnlineError(EOnlineErrorResult::MissingInterface));
+		});
+	}
+}
+
 FEOSVoiceChatUser* FOnlineSubsystemEOS::GetEOSVoiceChatUserInterface(const FUniqueNetId& LocalUserId)
 {
 	FEOSVoiceChatUser* Result = nullptr;
diff --git a/Engine/Plugins/Online/OnlineSubsystemEOS/Source/OnlineSubsystemEOS/Private/OnlineSubsystemEOS.h b/Engine/Plugins/Online/OnlineSubsystemEOS/Source/OnlineSubsystemEOS/Private/OnlineSubsystemEOS.h
--- a/Engine/Plugins/Online/OnlineSubsystemEOS/Source/OnlineSubsystemEOS/Private/OnlineSubsystemEOS.h
+++ b/Engine/Plugins/Online/OnlineSubsystemEOS/Source/OnlineSubsystemEOS/Private/OnlineSubsystemEOS.h
@@ -73,6 +73,7 @@ public:
 // IOnlineSubsystemEOS
 	virtual IVoiceChatUser* GetVoiceChatUserInterface(const FUniqueNetId& LocalUserId) override;
 	virtual IEOSPlatformHandlePtr GetEOSPlatformHandle() const override { return EOSPlatformHandle; };
+	virtual void QueryUniqueNetId(int32 LocalUserNum, const EOS_ProductUserId& ProductUserId, const FOnQueryUniqueNetIdComplete & Callback) override;
 
 // IOnlineSubsystem
 	virtual IOnlineSessionPtr GetSessionInterface() const override;
diff --git a/Engine/Plugins/Online/OnlineSubsystemEOS/Source/OnlineSubsystemEOS/Private/UserManagerEOS.cpp b/Engine/Plugins/Online/OnlineSubsystemEOS/Source/OnlineSubsystemEOS/Private/UserManagerEOS.cpp
--- a/Engine/Plugins/Online/OnlineSubsystemEOS/Source/OnlineSubsystemEOS/Private/UserManagerEOS.cpp
+++ b/Engine/Plugins/Online/OnlineSubsystemEOS/Source/OnlineSubsystemEOS/Private/UserManagerEOS.cpp
@@ -2283,7 +2283,7 @@ typedef TEOSCallback<EOS_Connect_OnQueryProductUserIdMappingsCallback, EOS_Conne
  *
  * @return true if the operation was successful, false otherwise
  */
-bool FUserManagerEOS::GetEpicAccountIdFromProductUserId(int32 LocalUserNum, const EOS_ProductUserId& ProductUserId, EOS_EpicAccountId& OutEpicAccountId) const
+bool FUserManagerEOS::GetEpicAccountIdFromProductUserId(int32 LocalUserNum, const EOS_ProductUserId& ProductUserId, EOS_EpicAccountId& OutEpicAccountId, FOnlineError& Error) const
 {
 	bool bResult = false;
 
@@ -2306,6 +2306,8 @@ bool FUserManagerEOS::GetEpicAccountIdFromProductUserId(int32 LocalUserNum, cons
 	else
 	{
 		UE_LOG_ONLINE(Verbose, TEXT("[FUserManagerEOS::GetEpicAccountIdFromProductUserId] EOS_Connect_GetProductUserIdMapping not successful for ProductUserId (%s). Finished with EOS_EResult %s"), *LexToString(ProductUserId), *LexToString(Result));
+		
+		Error = FOnlineError(LexToString(Result));
 	}
 
 	return bResult;
@@ -2315,8 +2317,8 @@ void FUserManagerEOS::ResolveUniqueNetId(int32 LocalUserNum, const EOS_ProductUs
 {
 	TArray<EOS_ProductUserId> ProductUserIds = { const_cast<EOS_ProductUserId>(ProductUserId) };
 
-	FResolveUniqueNetIdsCallback GroupCallback = [ProductUserId, OriginalCallback = Callback](TMap<EOS_ProductUserId, FUniqueNetIdEOSRef> ResolvedUniqueNetIds) {
-		OriginalCallback(ResolvedUniqueNetIds[ProductUserId]);
+	FResolveUniqueNetIdsCallback GroupCallback = [ProductUserId, OriginalCallback = Callback](TMap<EOS_ProductUserId, FUniqueNetIdEOSRef> ResolvedUniqueNetIds, const FOnlineError& Error) {
+		OriginalCallback(ResolvedUniqueNetIds[ProductUserId], Error);
 	};
 
 	ResolveUniqueNetIds(LocalUserNum, ProductUserIds, GroupCallback);
@@ -2327,12 +2329,14 @@ void FUserManagerEOS::ResolveUniqueNetIds(int32 LocalUserNum, const TArray<EOS_P
 	TMap<EOS_ProductUserId, FUniqueNetIdEOSRef> ResolvedUniqueNetIds;
 	TArray<EOS_ProductUserId> ProductUserIdsToResolve;
 
+	FOnlineError Error; // Defaults to Success
+
 	for (const EOS_ProductUserId& ProductUserId : ProductUserIds)
 	{
 		EOS_EpicAccountId EpicAccountId;
 
 		// We check first if the Product User Id has already been queried, which would allow us to retrieve its Epic Account Id directly
-		if (GetEpicAccountIdFromProductUserId(LocalUserNum, ProductUserId, EpicAccountId))
+		if (GetEpicAccountIdFromProductUserId(LocalUserNum, ProductUserId, EpicAccountId, Error))
 		{
 			const FUniqueNetIdEOSRef UniqueNetId = FUniqueNetIdEOSRegistry::FindOrAdd(EpicAccountId, ProductUserId);
 
@@ -2355,10 +2359,12 @@ void FUserManagerEOS::ResolveUniqueNetIds(int32 LocalUserNum, const TArray<EOS_P
 		QueryProductUserIdMappingsOptions.ProductUserIdCount = ProductUserIdsToResolve.Num();
 
 		FConnectQueryProductUserIdMappingsCallback* CallbackObj = new FConnectQueryProductUserIdMappingsCallback(FUserManagerEOSConstWeakPtr(AsShared()));
-		CallbackObj->CallbackLambda = [this, LocalUserNum, ProductUserIdsToResolve, ResolvedUniqueNetIds = MoveTemp(ResolvedUniqueNetIds), Callback](const EOS_Connect_QueryProductUserIdMappingsCallbackInfo* Data) mutable
+		CallbackObj->CallbackLambda = [this, LocalUserNum, ProductUserIdsToResolve, ResolvedUniqueNetIds = MoveTemp(ResolvedUniqueNetIds), Callback, Error](const EOS_Connect_QueryProductUserIdMappingsCallbackInfo* Data) mutable
 		{
 			if (Data->ResultCode != EOS_EResult::EOS_Success)
 			{
+				Error = FOnlineError(LexToString(Data->ResultCode));
+
 				UE_LOG_ONLINE(Verbose, TEXT("[FUserManagerEOS::ResolveUniqueNetIds] EOS_Connect_QueryProductUserIdMappings not successful for user (%s). Finished with EOS_EResult %s."), *LexToString(Data->LocalUserId), *LexToString(Data->ResultCode));
 			}
 
@@ -2366,38 +2372,38 @@ void FUserManagerEOS::ResolveUniqueNetIds(int32 LocalUserNum, const TArray<EOS_P
 			{
 				EOS_EpicAccountId EpicAccountId = nullptr;
 
-				GetEpicAccountIdFromProductUserId(LocalUserNum, ProductUserId, EpicAccountId);
+				GetEpicAccountIdFromProductUserId(LocalUserNum, ProductUserId, EpicAccountId, Error);
 
 				const FUniqueNetIdEOSRef UniqueNetId = FUniqueNetIdEOSRegistry::FindOrAdd(EpicAccountId, ProductUserId);
 
 				ResolvedUniqueNetIds.Add(ProductUserId, UniqueNetId);
 			}
 
-			Callback(ResolvedUniqueNetIds);
+			Callback(ResolvedUniqueNetIds, Error);
 		};
 
 		EOS_Connect_QueryProductUserIdMappings(EOSSubsystem->ConnectHandle, &QueryProductUserIdMappingsOptions, CallbackObj, CallbackObj->GetCallbackPtr());
 	}
 	else
 	{
-		EOSSubsystem->ExecuteNextTick([Callback, ResolvedUniqueNetIds]()
+		EOSSubsystem->ExecuteNextTick([Callback, ResolvedUniqueNetIds, Error]()
 			{
-				Callback(ResolvedUniqueNetIds);
+				Callback(ResolvedUniqueNetIds, Error);
 			});
 	}
 }
 
 void FUserManagerEOS::ResolveUniqueNetId(int32 LocalUserNum, const EOS_EpicAccountId& EpicAccountId, const FResolveUniqueNetIdCallback& Callback)
 {
-	const FResolveEpicAccountIdsCallback GroupCallback = [LocalUserNum, EpicAccountId, OriginalCallback = Callback](bool bWasSuccessful, TMap<EOS_EpicAccountId, FUniqueNetIdEOSRef> ResolvedUniqueNetIds, FString ErrorStr) {
+	const FResolveEpicAccountIdsCallback GroupCallback = [LocalUserNum, EpicAccountId, OriginalCallback = Callback](bool bWasSuccessful, TMap<EOS_EpicAccountId, FUniqueNetIdEOSRef> ResolvedUniqueNetIds, const FOnlineError& Error) {
 		if (bWasSuccessful)
 		{
-			OriginalCallback(ResolvedUniqueNetIds[EpicAccountId]);
+			OriginalCallback(ResolvedUniqueNetIds[EpicAccountId], Error);
 		}
 		else
 		{
-			UE_LOG_ONLINE(Error, TEXT("Unable to ResolveUniqueNetId for %d  %s"), LocalUserNum, *ErrorStr);
-			OriginalCallback(FUniqueNetIdEOS::EmptyId());
+			UE_LOG_ONLINE(Error, TEXT("Unable to ResolveUniqueNetId for %d  %s"), LocalUserNum, *Error.ToLogString());
+			OriginalCallback(FUniqueNetIdEOS::EmptyId(), Error);
 		}
 	};
 
@@ -2444,14 +2450,14 @@ void FUserManagerEOS::ResolveUniqueNetIds(int32 LocalUserNum, const TArray<EOS_E
 					}
 				}
 
-				Callback(bWasSuccessful, ResolvedUniqueNetIds, Error);
+				Callback(bWasSuccessful, ResolvedUniqueNetIds, FOnlineError(Error));
 			}));
 	}
 	else
 	{
 		EOSSubsystem->ExecuteNextTick([Callback, ResolvedUniqueNetIds = MoveTemp(ResolvedUniqueNetIds)]()
 			{
-				Callback(true, ResolvedUniqueNetIds, FString());
+				Callback(true, ResolvedUniqueNetIds, FOnlineError(EOnlineErrorResult::InvalidCreds));
 			});
 	}
 }
@@ -3002,7 +3008,7 @@ FOnlineFriendEOSRef FUserManagerEOS::AddFriend(int32 LocalUserNum, const FUnique
 
 void FUserManagerEOS::AddRemotePlayer(int32 LocalUserNum, EOS_EpicAccountId EpicAccountId, const FRemoteUserProcessedCallback& Callback)
 {
-	const FResolveUniqueNetIdCallback IdResolutionCallback = [this, LocalUserNum, EpicAccountId, Callback](FUniqueNetIdEOSRef ResolvedUniqueNetId) mutable
+	const FResolveUniqueNetIdCallback IdResolutionCallback = [this, LocalUserNum, EpicAccountId, Callback](FUniqueNetIdEOSRef ResolvedUniqueNetId, const FOnlineError& Error) mutable
 	{
 		const FOnlineUserEOSRef AttributeRef = MakeShareable(new FOnlineUserEOS(ResolvedUniqueNetId, *EOSSubsystem));
 
@@ -3017,7 +3023,7 @@ void FUserManagerEOS::AddRemotePlayer(int32 LocalUserNum, EOS_EpicAccountId Epic
 
 void FUserManagerEOS::AddRemotePlayers(int32 LocalUserNum, TArray<EOS_EpicAccountId> EpicAccountIds, const FRemoteUsersProcessedCallback& Callback)
 {
-	const FResolveEpicAccountIdsCallback IdResolutionCallback = [this, LocalUserNum, EpicAccountIds, Callback](bool bWasSuccessful, TMap<EOS_EpicAccountId, FUniqueNetIdEOSRef> ResolvedUniqueNetIds, FString ErrorStr) mutable
+	const FResolveEpicAccountIdsCallback IdResolutionCallback = [this, LocalUserNum, EpicAccountIds, Callback](bool bWasSuccessful, TMap<EOS_EpicAccountId, FUniqueNetIdEOSRef> ResolvedUniqueNetIds, const FOnlineError& Error) mutable
 	{
 		if (!ResolvedUniqueNetIds.IsEmpty())
 		{
@@ -3057,7 +3063,7 @@ void FUserManagerEOS::AddRemotePlayers(int32 LocalUserNum, TArray<EOS_EpicAccoun
 		{
 			// If the list is empty, we still need to execute the callback
 			TArray<FUniqueNetIdEOSRef> EmptyIdArray;
-			Callback(bWasSuccessful, EmptyIdArray, ErrorStr);
+			Callback(bWasSuccessful, EmptyIdArray, Error.ToLogString());
 		}			
 	};
 
diff --git a/Engine/Plugins/Online/OnlineSubsystemEOS/Source/OnlineSubsystemEOS/Private/UserManagerEOS.h b/Engine/Plugins/Online/OnlineSubsystemEOS/Source/OnlineSubsystemEOS/Private/UserManagerEOS.h
--- a/Engine/Plugins/Online/OnlineSubsystemEOS/Source/OnlineSubsystemEOS/Private/UserManagerEOS.h
+++ b/Engine/Plugins/Online/OnlineSubsystemEOS/Source/OnlineSubsystemEOS/Private/UserManagerEOS.h
@@ -364,10 +364,11 @@ public:
 	bool IsLocalUser(const EOS_EpicAccountId& EpicAccountId) const;
 	bool IsLocalUser(const EOS_ProductUserId& ProductUserId) const;
 
-	typedef TFunction<void(TMap<EOS_ProductUserId, FUniqueNetIdEOSRef> ResolvedUniqueNetIds)> FResolveUniqueNetIdsCallback;
-	typedef TFunction<void(bool bWasSuccessful, TMap<EOS_EpicAccountId, FUniqueNetIdEOSRef> ResolvedUniqueNetIds, FString ErrorStr)> FResolveEpicAccountIdsCallback;
-	typedef TFunction<void(FUniqueNetIdEOSRef ResolvedUniqueNetId)> FResolveUniqueNetIdCallback;
-	bool GetEpicAccountIdFromProductUserId(int32 LocalUserNum, const EOS_ProductUserId& ProductUserId, EOS_EpicAccountId& OutEpicAccountId) const;
+	using FResolveUniqueNetIdsCallback = TFunction<void(TMap<EOS_ProductUserId, FUniqueNetIdEOSRef> ResolvedUniqueNetIds, const FOnlineError& Error)>;
+	using FResolveEpicAccountIdsCallback = TFunction<void(bool bWasSuccessful, TMap<EOS_EpicAccountId, FUniqueNetIdEOSRef> ResolvedUniqueNetIds, const FOnlineError& Error)>;
+	using FResolveUniqueNetIdCallback = TFunction<void(FUniqueNetIdEOSRef ResolvedUniqueNetId, const FOnlineError& Error)>;
+
+	bool GetEpicAccountIdFromProductUserId(int32 LocalUserNum, const EOS_ProductUserId& ProductUserId, EOS_EpicAccountId& OutEpicAccountId, FOnlineError& Error) const;
 	void ResolveUniqueNetId(int32 LocalUserNum, const EOS_ProductUserId& ProductUserId, const FResolveUniqueNetIdCallback& Callback) const;
 	void ResolveUniqueNetIds(int32 LocalUserNum, const TArray<EOS_ProductUserId>& ProductUserIds, const FResolveUniqueNetIdsCallback& Callback) const;
 	void ResolveUniqueNetId(int32 LocalUserNum, const EOS_EpicAccountId& EpicAccountId, const FResolveUniqueNetIdCallback& Callback);
diff --git a/Engine/Plugins/Online/OnlineSubsystemEOS/Source/OnlineSubsystemEOS/Public/IOnlineSubsystemEOS.h b/Engine/Plugins/Online/OnlineSubsystemEOS/Source/OnlineSubsystemEOS/Public/IOnlineSubsystemEOS.h
--- a/Engine/Plugins/Online/OnlineSubsystemEOS/Source/OnlineSubsystemEOS/Public/IOnlineSubsystemEOS.h
+++ b/Engine/Plugins/Online/OnlineSubsystemEOS/Source/OnlineSubsystemEOS/Public/IOnlineSubsystemEOS.h
@@ -3,6 +3,7 @@
 #pragma once
 
 #include "OnlineSubsystemImpl.h"
+#include "OnlineDelegateMacros.h"
 
 class FUniqueNetId;
 class IVoiceChatUser;
@@ -10,6 +11,11 @@ using IEOSPlatformHandlePtr = TSharedPtr<class IEOSPlatformHandle, ESPMode::Thre
 using IOnlinePlayerSanctionEOSPtr = TSharedPtr<class IOnlinePlayerSanctionEOS, ESPMode::ThreadSafe>;
 using IOnlinePlayerReportEOSPtr = TSharedPtr<class IOnlinePlayerReportEOS, ESPMode::ThreadSafe>;
 
+using FUniqueNetIdEOSRef = TSharedRef<const class FUniqueNetIdEOS>;
+using EOS_ProductUserId = struct EOS_ProductUserIdDetails*;
+
+DECLARE_DELEGATE_TwoParams(FOnQueryUniqueNetIdComplete, FUniqueNetIdEOSRef /*ResolvedUniqueNetId*/, const FOnlineError& /*Error*/);
+
 /**
  *	OnlineSubsystemEOS - Implementation of the online subsystem for EOS services
  */
@@ -24,4 +30,6 @@ public:
 	virtual IEOSPlatformHandlePtr GetEOSPlatformHandle() const = 0;
 	virtual IOnlinePlayerSanctionEOSPtr GetPlayerSanctionEOSInterface() const = 0;
 	virtual IOnlinePlayerReportEOSPtr GetPlayerReportEOSInterface() const = 0;
+
+	virtual void QueryUniqueNetId(int32 LocalUserNum, const EOS_ProductUserId& ProductUserId, const FOnQueryUniqueNetIdComplete & Callback) = 0;
 };
\ No newline at end of file
