[OnlineSubsystemEOS] Allow controlling whether RTC is enabled using config.

diff --git a/Engine/Plugins/Online/OnlineSubsystemEOS/Source/OnlineSubsystemEOS/Private/EOSSettings.cpp b/Engine/Plugins/Online/OnlineSubsystemEOS/Source/OnlineSubsystemEOS/Private/EOSSettings.cpp
--- a/Engine/Plugins/Online/OnlineSubsystemEOS/Source/OnlineSubsystemEOS/Private/EOSSettings.cpp
+++ b/Engine/Plugins/Online/OnlineSubsystemEOS/Source/OnlineSubsystemEOS/Private/EOSSettings.cpp
@@ -147,6 +147,7 @@ FEOSSettings::FEOSSettings()
 	, bPreferPersistentAuth(false)
 	, bUseEAS(false)
 	, bUseEOSConnect(false)
+	, bUseEOSRTC(true)
 	, bUseEOSSessions(false)
 PRAGMA_DISABLE_DEPRECATION_WARNINGS
 	, bMirrorStatsToEOS(false)
@@ -197,6 +198,7 @@ const FEOSSettings& UEOSSettings::ManualGetSettings()
 		GConfig->GetBool(INI_SECTION, TEXT("bPreferPersistentAuth"), GCachedSettings->bPreferPersistentAuth, GEngineIni);
 		GConfig->GetBool(INI_SECTION, TEXT("bUseEAS"), GCachedSettings->bUseEAS, GEngineIni);
 		GConfig->GetBool(INI_SECTION, TEXT("bUseEOSConnect"), GCachedSettings->bUseEOSConnect, GEngineIni);
+		GConfig->GetBool(INI_SECTION, TEXT("bUseEOSRTC"), GCachedSettings->bUseEOSRTC, GEngineIni);
 		GConfig->GetBool(INI_SECTION, TEXT("bUseEOSSessions"), GCachedSettings->bUseEOSSessions, GEngineIni);
 PRAGMA_DISABLE_DEPRECATION_WARNINGS
 		GConfig->GetBool(INI_SECTION, TEXT("bMirrorStatsToEOS"), GCachedSettings->bMirrorStatsToEOS, GEngineIni);
@@ -232,6 +234,7 @@ FEOSSettings UEOSSettings::ToNative() const
 	Native.bPreferPersistentAuth = bPreferPersistentAuth;
 	Native.bUseEAS = bUseEAS;
 	Native.bUseEOSConnect = bUseEOSConnect;
+	Native.bUseEOSRTC = bUseEOSRTC;
 	Native.bUseEOSSessions = bUseEOSSessions;
 PRAGMA_DISABLE_DEPRECATION_WARNINGS
 	Native.bMirrorStatsToEOS = bMirrorStatsToEOS;
diff --git a/Engine/Plugins/Online/OnlineSubsystemEOS/Source/OnlineSubsystemEOS/Private/OnlineSubsystemEOS.cpp b/Engine/Plugins/Online/OnlineSubsystemEOS/Source/OnlineSubsystemEOS/Private/OnlineSubsystemEOS.cpp
--- a/Engine/Plugins/Online/OnlineSubsystemEOS/Source/OnlineSubsystemEOS/Private/OnlineSubsystemEOS.cpp
+++ b/Engine/Plugins/Online/OnlineSubsystemEOS/Source/OnlineSubsystemEOS/Private/OnlineSubsystemEOS.cpp
@@ -244,7 +244,7 @@ bool FOnlineSubsystemEOS::PlatformCreate()
 	UE_EOS_CHECK_API_MISMATCH(EOS_PLATFORM_RTCOPTIONS_API_LATEST, 2);
 	RtcOptions.PlatformSpecificOptions = nullptr;
 	RtcOptions.BackgroundMode = EOSSettings.RTCBackgroundMode;
-	PlatformOptions.RTCOptions = &RtcOptions;
+	PlatformOptions.RTCOptions = EOSSettings.bUseEOSRTC ? &RtcOptions : nullptr;
 #endif
 
 	EOSPlatformHandle = EOSHelpersPtr->CreatePlatform(PlatformOptions);
diff --git a/Engine/Plugins/Online/OnlineSubsystemEOS/Source/OnlineSubsystemEOS/Private/Windows/WindowsEOSHelpers.cpp b/Engine/Plugins/Online/OnlineSubsystemEOS/Source/OnlineSubsystemEOS/Private/Windows/WindowsEOSHelpers.cpp
--- a/Engine/Plugins/Online/OnlineSubsystemEOS/Source/OnlineSubsystemEOS/Private/Windows/WindowsEOSHelpers.cpp
+++ b/Engine/Plugins/Online/OnlineSubsystemEOS/Source/OnlineSubsystemEOS/Private/Windows/WindowsEOSHelpers.cpp
@@ -22,7 +22,7 @@ IEOSPlatformHandlePtr FWindowsEOSHelpers::CreatePlatform(EOS_Platform_Options& P
 	WindowsRtcOptions.XAudio29DllPath = Utf8XAudioPath.Get();
 
 	EOS_Platform_RTCOptions* RTCOptions = const_cast<EOS_Platform_RTCOptions*>(PlatformOptions.RTCOptions);
-	if (ensure(RTCOptions))
+	if (RTCOptions)
 	{
 		RTCOptions->PlatformSpecificOptions = &WindowsRtcOptions;
 	}
diff --git a/Engine/Plugins/Online/OnlineSubsystemEOS/Source/OnlineSubsystemEOS/Public/EOSSettings.h b/Engine/Plugins/Online/OnlineSubsystemEOS/Source/OnlineSubsystemEOS/Public/EOSSettings.h
--- a/Engine/Plugins/Online/OnlineSubsystemEOS/Source/OnlineSubsystemEOS/Public/EOSSettings.h
+++ b/Engine/Plugins/Online/OnlineSubsystemEOS/Source/OnlineSubsystemEOS/Public/EOSSettings.h
@@ -89,6 +89,7 @@ PRAGMA_ENABLE_DEPRECATION_WARNINGS
 	bool bPreferPersistentAuth;
 	bool bUseEAS;
 	bool bUseEOSConnect;
+	bool bUseEOSRTC;
 	bool bUseEOSSessions;
 	UE_DEPRECATED(5.5, "EOSPlus is deprecated, and related settings will be removed soon.")
 	bool bMirrorStatsToEOS;
@@ -169,6 +170,10 @@ public:
 	UPROPERTY(Config, EditAnywhere, BlueprintReadOnly, Category="EOS Settings", DisplayName="Login to EOS Connect, and use EOS Game Services")
 	bool bUseEOSConnect = false;
 
+	/** Whether real-time chat is initialized when creating the EOS platform */
+	UPROPERTY(Config, EditAnywhere, BlueprintReadOnly, Category="EOS Settings", DisplayName="Enable Real-Time chat")
+	bool bUseEOSRTC = true;
+
 	/** Set to true to write stats to EOS as well as the default platform */
 	UE_DEPRECATED(5.5, "EOSPlus is deprecated, and related settings will be removed soon.")
 	UPROPERTY(Config, EditAnywhere, BlueprintReadOnly, Category="Crossplay Settings")
