Deprecation of bUseEOSSessions, which is part of the EOSPlus logic

diff --git a/Engine/Plugins/Online/OnlineSubsystemEOS/Source/OnlineSubsystemEOS/Public/EOSSettings.h b/Engine/Plugins/Online/OnlineSubsystemEOS/Source/OnlineSubsystemEOS/Public/EOSSettings.h
--- a/Engine/Plugins/Online/OnlineSubsystemEOS/Source/OnlineSubsystemEOS/Public/EOSSettings.h
+++ b/Engine/Plugins/Online/OnlineSubsystemEOS/Source/OnlineSubsystemEOS/Public/EOSSettings.h
@@ -90,6 +90,7 @@ PRAGMA_ENABLE_DEPRECATION_WARNINGS
 	bool bUseEAS;
 	bool bUseEOSConnect;
 	bool bUseEOSRTC;
+	UE_DEPRECATED(5.6, "EOSPlus is deprecated, and related settings will be removed soon.")
 	bool bUseEOSSessions;
 	UE_DEPRECATED(5.5, "EOSPlus is deprecated, and related settings will be removed soon.")
 	bool bMirrorStatsToEOS;
@@ -185,6 +186,7 @@ public:
 	bool bMirrorAchievementsToEOS = false;
 
 	/** Set to true to use EOS for session registration with data mirrored to the default platform */
+	UE_DEPRECATED(5.6, "EOSPlus is deprecated, and related settings will be removed soon.")
 	UPROPERTY(Config, EditAnywhere, BlueprintReadOnly, Category="Crossplay Settings", DisplayName="Use Crossplay Sessions")
 	bool bUseEOSSessions = false;
 
diff --git a/Engine/Plugins/Online/OnlineSubsystemEOS/Source/OnlineSubsystemEOSPlus/Private/OnlineSessionEOSPlus.cpp b/Engine/Plugins/Online/OnlineSubsystemEOS/Source/OnlineSubsystemEOSPlus/Private/OnlineSessionEOSPlus.cpp
--- a/Engine/Plugins/Online/OnlineSubsystemEOS/Source/OnlineSubsystemEOSPlus/Private/OnlineSessionEOSPlus.cpp
+++ b/Engine/Plugins/Online/OnlineSubsystemEOS/Source/OnlineSubsystemEOSPlus/Private/OnlineSessionEOSPlus.cpp
@@ -16,7 +16,9 @@ FOnlineSessionEOSPlus::FOnlineSessionEOSPlus(FOnlineSubsystemEOSPlus* InSubsyste
 	EOSSessionInterface = EOSPlus->EosOSS->GetSessionInterface();
 	check(BaseSessionInterface.IsValid() && EOSSessionInterface.IsValid());
 
+PRAGMA_DISABLE_DEPRECATION_WARNINGS
 	bUseEOSSessions = UEOSSettings::GetSettings().bUseEOSSessions;
+PRAGMA_ENABLE_DEPRECATION_WARNINGS
 
 	BaseSessionInterface->AddOnSessionUserInviteAcceptedDelegate_Handle(FOnSessionUserInviteAcceptedDelegate::CreateRaw(this, &FOnlineSessionEOSPlus::OnSessionUserInviteAcceptedBase));
 	BaseSessionInterface->AddOnSessionInviteReceivedDelegate_Handle(FOnSessionInviteReceivedDelegate::CreateRaw(this, &FOnlineSessionEOSPlus::OnSessionInviteReceivedBase));
