Removed erroneous setting of DisplayName in UserInfo type when using EPIC credentials

diff --git a/Engine/Plugins/Online/OnlineSubsystemEOS/Source/OnlineSubsystemEOS/Private/UserManagerEOS.cpp b/Engine/Plugins/Online/OnlineSubsystemEOS/Source/OnlineSubsystemEOS/Private/UserManagerEOS.cpp
--- a/Engine/Plugins/Online/OnlineSubsystemEOS/Source/OnlineSubsystemEOS/Private/UserManagerEOS.cpp
+++ b/Engine/Plugins/Online/OnlineSubsystemEOS/Source/OnlineSubsystemEOS/Private/UserManagerEOS.cpp
@@ -1384,6 +1384,7 @@ bool FUserManagerEOS::ConnectLoginEASLegacy(int32 LocalUserNum, EOS_EpicAccountI
 						UE_EOS_CHECK_API_MISMATCH(EOS_CONNECT_USERLOGININFO_API_LATEST, 2);
 						auto AuthTokenConverter = StringCast<UTF8CHAR>(*AuthToken.TokenString);
 						UserLoginInfo.NsaIdToken = (const char*)AuthTokenConverter.Get();
+						// UserLoginInfo.DisplayName should not be set when the Credentials Type is EPIC
 #endif
 
 						EOS_Connect_LoginOptions Options = { };
@@ -1470,6 +1471,7 @@ bool FUserManagerEOS::ConnectLoginEAS(int32 LocalUserNum, EOS_EpicAccountId Acco
 						UE_EOS_CHECK_API_MISMATCH(EOS_CONNECT_USERLOGININFO_API_LATEST, 2);
 						auto AuthTokenConverter = StringCast<UTF8CHAR>(*AuthToken.TokenString);
 						UserLoginInfo.NsaIdToken = (const char*)AuthTokenConverter.Get();
+						// UserLoginInfo.DisplayName should not be set when the Credentials Type is EPIC
 #endif
 
 						EOS_Connect_LoginOptions Options = { };
@@ -1574,10 +1576,9 @@ void FUserManagerEOS::AutoRefreshConnectLogin(int32 LocalUserNum)
 						EOS_Connect_UserLoginInfo UserLoginInfo = {};
 						UserLoginInfo.ApiVersion = 2;
 						UE_EOS_CHECK_API_MISMATCH(EOS_CONNECT_USERLOGININFO_API_LATEST, 2);
-						const FTCHARToUTF8 DisplayNameUtf8(*GetPlatformDisplayName(LocalUserNum));
-						UserLoginInfo.DisplayName = DisplayNameUtf8.Get();
 						auto AuthTokenConverter = StringCast<UTF8CHAR>(*AuthToken.TokenString);
 						UserLoginInfo.NsaIdToken = (const char*)AuthTokenConverter.Get();
+						// UserLoginInfo.DisplayName should not be set when the Credentials Type is EPIC
 
 						Options.UserLoginInfo = &UserLoginInfo;
 #endif
