FUserManagerEOS::Logout assumes bUseEAS=true and there's no code path when bUSeEAS=false

diff --git a/Engine/Plugins/Online/OnlineSubsystemEOS/Source/OnlineSubsystemEOS/Private/UserManagerEOS.cpp b/Engine/Plugins/Online/OnlineSubsystemEOS/Source/OnlineSubsystemEOS/Private/UserManagerEOS.cpp
--- a/Engine/Plugins/Online/OnlineSubsystemEOS/Source/OnlineSubsystemEOS/Private/UserManagerEOS.cpp
+++ b/Engine/Plugins/Online/OnlineSubsystemEOS/Source/OnlineSubsystemEOS/Private/UserManagerEOS.cpp
@@ -1832,7 +1832,8 @@ void FUserManagerEOS::FullLoginCallback(int32 LocalUserNum, EOS_EpicAccountId Ac
 	TriggerOnLoginStatusChangedDelegates(LocalUserNum, ELoginStatus::NotLoggedIn, ELoginStatus::LoggedIn, *LocalUser.UniqueNetId);
 }
 
-typedef TEOSCallback<EOS_Auth_OnLogoutCallback, EOS_Auth_LogoutCallbackInfo, FUserManagerEOS> FLogoutCallback;
+typedef TEOSCallback<EOS_Auth_OnLogoutCallback, EOS_Auth_LogoutCallbackInfo, FUserManagerEOS> FAuthLogoutCallback;
+typedef TEOSCallback<EOS_Connect_OnLogoutCallback, EOS_Connect_LogoutCallbackInfo, FUserManagerEOS> FConnectLogoutCallback;
 
 bool FUserManagerEOS::Logout(int32 LocalUserNum)
 {
@@ -1845,40 +1846,80 @@ bool FUserManagerEOS::Logout(int32 LocalUserNum)
 		return false;
 	}
 
-	FLogoutCallback* CallbackObj = new FLogoutCallback(AsWeak());
-	CallbackObj->CallbackLambda = [LocalUserNum, this](const EOS_Auth_LogoutCallbackInfo* Data)
+	FEOSSettings Settings = UEOSSettings::GetSettings();
+
+	if (Settings.bUseEAS)
+	{
+		CallEOSAuthLogout(LocalUserNum, UserId);
+	}
+	else
 	{
-		FDeletePersistentAuthCallback* DeleteAuthCallbackObj = new FDeletePersistentAuthCallback(AsWeak());
-		DeleteAuthCallbackObj->CallbackLambda = [this, LocalUserNum, LogoutResultCode = Data->ResultCode](const EOS_Auth_DeletePersistentAuthCallbackInfo* Data)
+		CallEOSConnectLogout(LocalUserNum, UserId);
+	}
+
+	return true;
+} 
+
+void FUserManagerEOS::CallEOSAuthLogout(int32 LocalUserNum, FUniqueNetIdEOSPtr UserId)
+{
+	FAuthLogoutCallback* CallbackObj = new FAuthLogoutCallback(AsWeak());
+	CallbackObj->CallbackLambda = [LocalUserNum, UserId, this](const EOS_Auth_LogoutCallbackInfo* Data)
 		{
-			if (LogoutResultCode == EOS_EResult::EOS_Success)
+			if (Data->ResultCode == EOS_EResult::EOS_Success)
 			{
-				RemoveLocalUser(LocalUserNum);
+				FDeletePersistentAuthCallback* DeleteAuthCallbackObj = new FDeletePersistentAuthCallback(AsWeak());
+				DeleteAuthCallbackObj->CallbackLambda = [this, LocalUserNum, UserId](const EOS_Auth_DeletePersistentAuthCallbackInfo* Data)
+					{
+						UE_LOG_ONLINE(Verbose, TEXT("EOS_Auth_DeletePersistentAuth Result=[%s]"), *LexToString(Data->ResultCode));
+						CallEOSConnectLogout(LocalUserNum, UserId);
+					};
 
-				TriggerOnLogoutCompleteDelegates(LocalUserNum, true);
+				EOS_Auth_DeletePersistentAuthOptions DeletePersistentAuthOptions;
+				DeletePersistentAuthOptions.ApiVersion = 2;
+				UE_EOS_CHECK_API_MISMATCH(EOS_AUTH_DELETEPERSISTENTAUTH_API_LATEST, 2);
+				DeletePersistentAuthOptions.RefreshToken = nullptr;
+				EOS_Auth_DeletePersistentAuth(EOSSubsystem->AuthHandle, &DeletePersistentAuthOptions, (void*)DeleteAuthCallbackObj, DeleteAuthCallbackObj->GetCallbackPtr());
 			}
 			else
 			{
+				// EOS_Auth_DeletePersistentAuth and EOS_Connect_Logout are not called on failure. The game should attempt to logout again
+				UE_LOG_ONLINE(Warning, TEXT("EOS_Auth_Logout Result=[%s]"), *LexToString(Data->ResultCode))
 				TriggerOnLogoutCompleteDelegates(LocalUserNum, false);
 			}
 		};
 
-		EOS_Auth_DeletePersistentAuthOptions DeletePersistentAuthOptions;
-		DeletePersistentAuthOptions.ApiVersion = 2;
-		UE_EOS_CHECK_API_MISMATCH(EOS_AUTH_DELETEPERSISTENTAUTH_API_LATEST, 2);
-		DeletePersistentAuthOptions.RefreshToken = nullptr;
-		EOS_Auth_DeletePersistentAuth(EOSSubsystem->AuthHandle, &DeletePersistentAuthOptions, (void*)DeleteAuthCallbackObj, DeleteAuthCallbackObj->GetCallbackPtr());
-	};
-
 	EOS_Auth_LogoutOptions LogoutOptions = { };
 	LogoutOptions.ApiVersion = 1;
 	UE_EOS_CHECK_API_MISMATCH(EOS_AUTH_LOGOUT_API_LATEST, 1);
 	LogoutOptions.LocalUserId = UserId->GetEpicAccountId();
 
 	EOS_Auth_Logout(EOSSubsystem->AuthHandle, &LogoutOptions, CallbackObj, CallbackObj->GetCallbackPtr());
+}
 
-	return true;
-} 
+void FUserManagerEOS::CallEOSConnectLogout(int32 LocalUserNum, FUniqueNetIdEOSPtr UserId)
+{
+	FConnectLogoutCallback* CallbackObj = new FConnectLogoutCallback(AsWeak());
+	CallbackObj->CallbackLambda = [LocalUserNum, this](const EOS_Connect_LogoutCallbackInfo* Data)
+		{
+			const bool bWasConnectLogoutSuccessful = Data->ResultCode == EOS_EResult::EOS_Success;
+
+			UE_CLOG_ONLINE(!bWasConnectLogoutSuccessful, Warning, TEXT("EOS_Connect_Logout Result=[%s]"), *LexToString(Data->ResultCode));
+
+			if (bWasConnectLogoutSuccessful)
+			{
+				RemoveLocalUser(LocalUserNum);
+			}
+
+			TriggerOnLogoutCompleteDelegates(LocalUserNum, bWasConnectLogoutSuccessful);
+		};
+
+	EOS_Connect_LogoutOptions LogoutOptions = { };
+	LogoutOptions.ApiVersion = 1;
+	UE_EOS_CHECK_API_MISMATCH(EOS_CONNECT_LOGOUT_API_LATEST, 1);
+	LogoutOptions.LocalUserId = UserId->GetProductUserId();
+
+	EOS_Connect_Logout(EOSSubsystem->ConnectHandle, &LogoutOptions, CallbackObj, CallbackObj->GetCallbackPtr());
+}
 
 bool FUserManagerEOS::AutoLoginLegacy(int32 LocalUserNum)
 {
diff --git a/Engine/Plugins/Online/OnlineSubsystemEOS/Source/OnlineSubsystemEOS/Private/UserManagerEOS.h b/Engine/Plugins/Online/OnlineSubsystemEOS/Source/OnlineSubsystemEOS/Private/UserManagerEOS.h
--- a/Engine/Plugins/Online/OnlineSubsystemEOS/Source/OnlineSubsystemEOS/Private/UserManagerEOS.h
+++ b/Engine/Plugins/Online/OnlineSubsystemEOS/Source/OnlineSubsystemEOS/Private/UserManagerEOS.h
@@ -421,6 +421,9 @@ private:
 	void OnEOSAuthLoginComplete(int32 LocalUserNum, const EOS_ELoginCredentialType LoginCredentialType, const bool bIsAutoLogin, const EOS_Auth_LoginCallbackInfo* Data);
 	void CallEOSConnectLogin(int32 LocalUserNum, const FOnlineAccountCredentials& Credentials);
 
+	void CallEOSAuthLogout(int32 LocalUserNum, FUniqueNetIdEOSPtr UserId);
+	void CallEOSConnectLogout(int32 LocalUserNum, FUniqueNetIdEOSPtr UserId);
+	
 	FLocalUserEOS& AddLocalUser(int32 LocalUserNum);
 	FLocalUserEOS& UpdateLocalUser(int32 LocalUserNum, EOS_EpicAccountId EpicAccountId, EOS_ProductUserId UserId);
 	void RemoveLocalUser(int32 LocalUserNum);
