Fix error path in FUserManagerEOS::ReadFriendsList when user failed to authorize

diff --git a/Engine/Plugins/Online/OnlineSubsystemEOS/Source/OnlineSubsystemEOS/Private/UserManagerEOS.cpp b/Engine/Plugins/Online/OnlineSubsystemEOS/Source/OnlineSubsystemEOS/Private/UserManagerEOS.cpp
--- a/Engine/Plugins/Online/OnlineSubsystemEOS/Source/OnlineSubsystemEOS/Private/UserManagerEOS.cpp
+++ b/Engine/Plugins/Online/OnlineSubsystemEOS/Source/OnlineSubsystemEOS/Private/UserManagerEOS.cpp
@@ -2690,26 +2690,33 @@ bool FUserManagerEOS::ReadFriendsList(int32 LocalUserNum, const FString& ListNam
 
 			const FRemoteUsersProcessedCallback Callback = [this, LocalUserNum, FriendCount](bool bWasSuccessful, TArray<FUniqueNetIdEOSRef> RemoteUserNetIds, const FString& ErrorStr)
 			{
-				if (bWasSuccessful)
+				if (IsLocalUserValid(LocalUserNum))
 				{
-					FFriendsListEOSPtr FriendsListPtr = GetLocalUserChecked(LocalUserNum).FriendsList;
-					if (FriendsListPtr)
+					if (bWasSuccessful)
 					{
-						FriendsListPtr->Empty(FriendCount);
-					}
-					else
-					{
-						GetLocalUserChecked(LocalUserNum).FriendsList = MakeShareable(new FFriendsListEOS(LocalUserNum, GetLocalUniqueNetIdEOS(LocalUserNum).ToSharedRef()));
-					}
+						FFriendsListEOSPtr FriendsListPtr = GetLocalUserChecked(LocalUserNum).FriendsList;
+						if (FriendsListPtr)
+						{
+							FriendsListPtr->Empty(FriendCount);
+						}
+						else
+						{
+							GetLocalUserChecked(LocalUserNum).FriendsList = MakeShared<FFriendsListEOS>(LocalUserNum, GetLocalUniqueNetIdEOS(LocalUserNum).ToSharedRef());
+						}
 
-					for (const FUniqueNetIdEOSRef& RemoteUserNetId : RemoteUserNetIds)
-					{
-						AddFriend(LocalUserNum, *RemoteUserNetId);
+						for (const FUniqueNetIdEOSRef& RemoteUserNetId : RemoteUserNetIds)
+						{
+							AddFriend(LocalUserNum, *RemoteUserNetId);
+						}
 					}
-				}
 
-				// We trigger all the corresponding delegates
-				ProcessReadFriendsListComplete(LocalUserNum, bWasSuccessful, ErrorStr);
+					// We trigger all the corresponding delegates
+					ProcessReadFriendsListComplete(LocalUserNum, bWasSuccessful, ErrorStr);
+				}
+				else
+				{
+					UE_LOG_ONLINE_FRIEND(Verbose, TEXT("[FUserManagerEOS::ReadFriendsList] User is no longer valid. Abandoning callback."));
+				}
 			};
 
 			AddRemotePlayers(LocalUserNum, FriendEasIds, Callback);
