Enable presence on join for EOS Sessions in OnlineSubsystemEOS

diff --git a/Engine/Plugins/Online/OnlineSubsystemEOS/Source/OnlineSubsystemEOS/Private/OnlineSessionEOS.cpp b/Engine/Plugins/Online/OnlineSubsystemEOS/Source/OnlineSubsystemEOS/Private/OnlineSessionEOS.cpp
--- a/Engine/Plugins/Online/OnlineSubsystemEOS/Source/OnlineSubsystemEOS/Private/OnlineSessionEOS.cpp
+++ b/Engine/Plugins/Online/OnlineSubsystemEOS/Source/OnlineSubsystemEOS/Private/OnlineSessionEOS.cpp
@@ -2875,7 +2875,7 @@ bool FOnlineSessionEOS::JoinSession(int32 PlayerNum, FName SessionName, const FO
 				}
 				else
 				{
-					Return = JoinEOSSession(PlayerNum, Session, &DesiredSession.Session);
+					Return = JoinEOSSession(PlayerNum, Session);
 				}
 			}
 			else
@@ -2939,7 +2939,7 @@ struct FJoinSessionOptions :
 
 typedef TEOSCallback<EOS_Sessions_OnJoinSessionCallback, EOS_Sessions_JoinSessionCallbackInfo, FOnlineSessionEOS> FJoinSessionCallback;
 
-uint32 FOnlineSessionEOS::JoinEOSSession(int32 PlayerNum, FNamedOnlineSession* Session, const FOnlineSession* SearchSession)
+uint32 FOnlineSessionEOS::JoinEOSSession(int32 PlayerNum, FNamedOnlineSession* Session)
 {
 	if (!Session->SessionInfo.IsValid())
 	{
@@ -2959,10 +2959,6 @@ uint32 FOnlineSessionEOS::JoinEOSSession(int32 PlayerNum, FNamedOnlineSession* S
 		return ONLINE_FAIL;
 	}
 
-	// Copy the session info over
-	TSharedPtr<const FOnlineSessionInfoEOS> SearchSessionInfo = StaticCastSharedPtr<const FOnlineSessionInfoEOS>(SearchSession->SessionInfo);
-	EOSSessionInfo->HostAddr = SearchSessionInfo->HostAddr->Clone();
-
 	Session->SessionState = EOnlineSessionState::Pending;
 
 	FName SessionName = Session->SessionName;
@@ -2996,6 +2992,7 @@ uint32 FOnlineSessionEOS::JoinEOSSession(int32 PlayerNum, FNamedOnlineSession* S
 	FJoinSessionOptions Options(TCHAR_TO_UTF8(*Session->SessionName.ToString()));
 	Options.LocalUserId = ProductUserId;
 	Options.SessionHandle = EOSSessionInfo->SessionHandle->SessionDetailsHandle;
+	Options.bPresenceEnabled = Session->SessionSettings.bUsesPresence;
 	EOS_Sessions_JoinSession(EOSSubsystem->SessionsHandle, &Options, CallbackObj, CallbackObj->GetCallbackPtr());
 
 	return ONLINE_IO_PENDING;
diff --git a/Engine/Plugins/Online/OnlineSubsystemEOS/Source/OnlineSubsystemEOS/Private/OnlineSessionEOS.h b/Engine/Plugins/Online/OnlineSubsystemEOS/Source/OnlineSubsystemEOS/Private/OnlineSessionEOS.h
--- a/Engine/Plugins/Online/OnlineSubsystemEOS/Source/OnlineSubsystemEOS/Private/OnlineSessionEOS.h
+++ b/Engine/Plugins/Online/OnlineSubsystemEOS/Source/OnlineSubsystemEOS/Private/OnlineSessionEOS.h
@@ -285,7 +285,7 @@ private:
 	TArray<TSharedRef<FSessionDetailsEOS>> SessionSearchResultsPendingIdResolution;
 
 	uint32 CreateEOSSession(int32 HostingPlayerNum, FNamedOnlineSession* Session);
-	uint32 JoinEOSSession(int32 PlayerNum, FNamedOnlineSession* Session, const FOnlineSession* SearchSession);
+	uint32 JoinEOSSession(int32 PlayerNum, FNamedOnlineSession* Session);
 	uint32 StartEOSSession(FNamedOnlineSession* Session);
 	uint32 UpdateEOSSession(FNamedOnlineSession* Session);
 	uint32 EndEOSSession(FNamedOnlineSession* Session);
