Fix for Horde issue 946600

Refactored to stop dereferencing null pointers

diff --git a/Engine/Plugins/Online/OnlineServicesEOS/Source/Private/Online/CommerceEOS.cpp b/Engine/Plugins/Online/OnlineServicesEOS/Source/Private/Online/CommerceEOS.cpp
--- a/Engine/Plugins/Online/OnlineServicesEOS/Source/Private/Online/CommerceEOS.cpp
+++ b/Engine/Plugins/Online/OnlineServicesEOS/Source/Private/Online/CommerceEOS.cpp
@@ -437,13 +437,16 @@ TOnlineAsyncOpHandle<FCommerceQueryEntitlements> FCommerceEOS::QueryEntitlements
 			Options.LocalUserId = GetEpicAccountIdChecked(Params.LocalAccountId);
 			Options.SandboxIdsCount = 1;
 
-			IEOSSDKManager* Manager = IEOSSDKManager::Get();
-			const FString& PlatformConfigName = StaticCast<FOnlineServicesEpicCommon*>(&Services)->GetEOSPlatformHandle()->GetConfigName();
-			const FEOSSDKPlatformConfig* Config = Manager ? Manager->GetPlatformConfig(PlatformConfigName) : nullptr;
-
-			const FTCHARToUTF8& Converter = FTCHARToUTF8(*Config->SandboxId);
-			const char* SandboxId = Converter.Get();
-			Options.SandboxIds = &SandboxId;
+			if (IEOSSDKManager* Manager = IEOSSDKManager::Get())
+			{
+				const FString& PlatformConfigName = StaticCast<FOnlineServicesEpicCommon*>(&Services)->GetEOSPlatformHandle()->GetConfigName();
+				if (const FEOSSDKPlatformConfig* Config = Manager->GetPlatformConfig(PlatformConfigName))
+				{
+					const FTCHARToUTF8& Converter = FTCHARToUTF8(*Config->SandboxId);
+					const char* SandboxId = Converter.Get();
+					Options.SandboxIds = &SandboxId;
+				}
+			}
 			
 			EOS_Async(EOS_Ecom_QueryOwnershipBySandboxIds, EcomHandle, Options, MoveTemp(Promise));
 		})
diff --git a/Engine/Plugins/Online/OnlineSubsystemEOS/Source/OnlineSubsystemEOS/Private/OnlineStoreEOS.cpp b/Engine/Plugins/Online/OnlineSubsystemEOS/Source/OnlineSubsystemEOS/Private/OnlineStoreEOS.cpp
--- a/Engine/Plugins/Online/OnlineSubsystemEOS/Source/OnlineSubsystemEOS/Private/OnlineStoreEOS.cpp
+++ b/Engine/Plugins/Online/OnlineSubsystemEOS/Source/OnlineSubsystemEOS/Private/OnlineStoreEOS.cpp
@@ -379,13 +379,16 @@ void FOnlineStoreEOS::QueryReceipts(const FUniqueNetId& UserId, bool bRestoreRec
 		Options.LocalUserId = Data->LocalUserId;
 		Options.SandboxIdsCount = 1;
 
-		IEOSSDKManager* Manager = IEOSSDKManager::Get();
-		const FString& PlatformConfigName = EOSSubsystem->GetEOSPlatformHandle()->GetConfigName();
-		const FEOSSDKPlatformConfig* Config = Manager ? Manager->GetPlatformConfig(PlatformConfigName) : nullptr;
-
-		const FTCHARToUTF8& Converter = FTCHARToUTF8(*Config->SandboxId);
-		const char* SandboxId = Converter.Get();
-		Options.SandboxIds = &SandboxId;
+		if (IEOSSDKManager* Manager = IEOSSDKManager::Get())
+		{
+			const FString& PlatformConfigName = EOSSubsystem->GetEOSPlatformHandle()->GetConfigName();
+			if (const FEOSSDKPlatformConfig* Config = Manager->GetPlatformConfig(PlatformConfigName))
+			{
+				const FTCHARToUTF8& Converter = FTCHARToUTF8(*Config->SandboxId);
+				const char* SandboxId = Converter.Get();
+				Options.SandboxIds = &SandboxId;
+			}
+		}		
 
 		FQueryOwnershipBySandboxIdsCallback* CallbackObj = new FQueryOwnershipBySandboxIdsCallback(FOnlineStoreEOSWeakPtr(AsShared()));
 		CallbackObj->CallbackLambda = [this, OnComplete = FOnQueryReceiptsComplete(OnComplete)](const EOS_Ecom_QueryOwnershipBySandboxIdsCallbackInfo* Data)
