Removed logic filtering EOS leaderboard results that did not have epic ids

diff --git a/Engine/Plugins/Online/OnlineSubsystemEOS/Source/OnlineSubsystemEOS/Private/OnlineLeaderboardsEOS.cpp b/Engine/Plugins/Online/OnlineSubsystemEOS/Source/OnlineSubsystemEOS/Private/OnlineLeaderboardsEOS.cpp
--- a/Engine/Plugins/Online/OnlineSubsystemEOS/Source/OnlineSubsystemEOS/Private/OnlineLeaderboardsEOS.cpp
+++ b/Engine/Plugins/Online/OnlineSubsystemEOS/Source/OnlineSubsystemEOS/Private/OnlineLeaderboardsEOS.cpp
@@ -257,14 +257,13 @@ bool FOnlineLeaderboardsEOS::ReadLeaderboardsAroundRank(int32 Rank, uint32 Range
 	FCStringAnsi::Strncpy(LeaderboardId, TCHAR_TO_UTF8(*ReadObject->LeaderboardName), EOS_OSS_STRING_BUFFER_LENGTH);
 
 	FQueryLeaderboardCallback* CallbackObj = new FQueryLeaderboardCallback(FOnlineLeaderboardsEOSWeakPtr(AsShared()));
-	FOnlineLeaderboardReadRef LambdaReadObject = ReadObject;
-	CallbackObj->CallbackLambda = [this, LambdaReadObject, StartIndex, EndIndex](const EOS_Leaderboards_OnQueryLeaderboardRanksCompleteCallbackInfo* Data)
+	CallbackObj->CallbackLambda = [this, ReadObject, StartIndex, EndIndex](const EOS_Leaderboards_OnQueryLeaderboardRanksCompleteCallbackInfo* Data)
 	{
 		bool bWasSuccessful = Data->ResultCode == EOS_EResult::EOS_Success;
 		if (!bWasSuccessful)
 		{
 			UE_LOG_ONLINE_LEADERBOARD(Error, TEXT("EOS_Leaderboards_QueryLeaderboardRanks() failed with EOS result code (%s)"), *LexToString(Data->ResultCode));
-			LambdaReadObject->ReadState = EOnlineAsyncTaskState::Failed;
+			ReadObject->ReadState = EOnlineAsyncTaskState::Failed;
 			TriggerOnLeaderboardReadCompleteDelegates(false);
 			return;
 		}
@@ -276,24 +275,19 @@ bool FOnlineLeaderboardsEOS::ReadLeaderboardsAroundRank(int32 Rank, uint32 Range
 		// Handle fewer entries than our start index
 		if (LeaderboardCount <= StartIndex)
 		{
-			LambdaReadObject->ReadState = EOnlineAsyncTaskState::Done;
+			ReadObject->ReadState = EOnlineAsyncTaskState::Done;
 			TriggerOnLeaderboardReadCompleteDelegates(false);
 			return;
 		}
 		// Handle fewer entries than our ending page index
 		uint32 NewEndIndex = FMath::Clamp<uint32>(EndIndex, 0, LeaderboardCount - 1);
 
+		// We get the list of users and resolve it first
 		EOS_Leaderboards_CopyLeaderboardRecordByIndexOptions CopyOptions = { };
 		CopyOptions.ApiVersion = 2;
 		UE_EOS_CHECK_API_MISMATCH(EOS_LEADERBOARDS_COPYLEADERBOARDRECORDBYINDEX_API_LATEST, 2);
 
-		EOS_Connect_GetProductUserIdMappingOptions Options = { };
-		Options.ApiVersion = 1;
-		UE_EOS_CHECK_API_MISMATCH(EOS_CONNECT_GETPRODUCTUSERIDMAPPING_API_LATEST, 1);
-		Options.AccountIdType = EOS_EExternalAccountType::EOS_EAT_EPIC;
-		Options.LocalUserId = EOSSubsystem->UserManager->GetLocalProductUserId(0);
-
-		// Allow for us reading a page out of the whole stats range returned
+		TArray<EOS_ProductUserId> UsersToResolve;
 		for (uint32 Index = StartIndex; Index <= NewEndIndex; Index++)
 		{
 			CopyOptions.LeaderboardRecordIndex = Index;
@@ -302,30 +296,43 @@ bool FOnlineLeaderboardsEOS::ReadLeaderboardsAroundRank(int32 Rank, uint32 Range
 			EOS_EResult Result = EOS_Leaderboards_CopyLeaderboardRecordByIndex(EOSSubsystem->LeaderboardsHandle, &CopyOptions, &Record);
 			if (Result == EOS_EResult::EOS_Success)
 			{
-				char EpicIdStr[EOS_CONNECT_EXTERNAL_ACCOUNT_ID_MAX_LENGTH+1];
-				int32 EpicIdStrSize = sizeof(EpicIdStr);
-				Options.TargetProductUserId = Record->UserId;
+				UsersToResolve.Add(Record->UserId);
+			}
+		}
+
+		EOSSubsystem->UserManager->ResolveUniqueNetIds(0, UsersToResolve, [this, StartIndex, NewEndIndex, ReadObject](TMap<EOS_ProductUserId, FUniqueNetIdEOSRef> ResolvedUniqueNetIds, const FOnlineError& Error) mutable
+			{
+				// Then we copy the rest of the data
+				EOS_Leaderboards_CopyLeaderboardRecordByIndexOptions CopyOptions = { };
+				CopyOptions.ApiVersion = 2;
+				UE_EOS_CHECK_API_MISMATCH(EOS_LEADERBOARDS_COPYLEADERBOARDRECORDBYINDEX_API_LATEST, 2);
 
-				Result = EOS_Connect_GetProductUserIdMapping(EOSSubsystem->ConnectHandle, &Options, EpicIdStr, &EpicIdStrSize);
-				if (Result == EOS_EResult::EOS_Success)
+				for (uint32 Index = StartIndex; Index <= NewEndIndex; Index++)
 				{
-					FString Nickname = UTF8_TO_TCHAR(Record->UserDisplayName);
-					if(Nickname.IsEmpty())
+					CopyOptions.LeaderboardRecordIndex = Index;
+
+					EOS_Leaderboards_LeaderboardRecord* Record = nullptr;
+					EOS_EResult Result = EOS_Leaderboards_CopyLeaderboardRecordByIndex(EOSSubsystem->LeaderboardsHandle, &CopyOptions, &Record);
+					if (Result == EOS_EResult::EOS_Success)
 					{
-						Nickname = TEXT("Unknown Player");
+						FString Nickname = UTF8_TO_TCHAR(Record->UserDisplayName);
+						if (Nickname.IsEmpty())
+						{
+							Nickname = TEXT("Unknown Player");
+						}
+						
+						if (const FUniqueNetIdEOSRef* NetId = ResolvedUniqueNetIds.Find(Record->UserId))
+						{
+							FOnlineStatsRow* Row = new(ReadObject->Rows) FOnlineStatsRow(Nickname, *NetId);
+							Row->Rank = Record->Rank;
+							Row->Columns.Add(ReadObject->SortedColumn, FVariantData(Record->Score));
+						}
 					}
-					const EOS_EpicAccountId AccountId = EOS_EpicAccountId_FromString(EpicIdStr);
-					const FUniqueNetIdEOSRef NetId = FUniqueNetIdEOSRegistry::FindOrAdd(AccountId, Record->UserId);
-
-					FOnlineStatsRow* Row = new(LambdaReadObject->Rows) FOnlineStatsRow(Nickname, NetId);
-					Row->Rank = Record->Rank;
-					Row->Columns.Add(LambdaReadObject->SortedColumn, FVariantData(Record->Score));
 				}
-			}
-		}
 
-		LambdaReadObject->ReadState = EOnlineAsyncTaskState::Done;
-		TriggerOnLeaderboardReadCompleteDelegates(true);
+				ReadObject->ReadState = EOnlineAsyncTaskState::Done;
+				TriggerOnLeaderboardReadCompleteDelegates(true);
+			});
 	};
 
 	ReadObject->ReadState = EOnlineAsyncTaskState::InProgress;
