Check LoginStatus in SocketSubsystemEOS before calling EOS P2P APIs

diff --git a/Engine/Plugins/Online/OnlineServicesEOSGS/Source/Private/Online/SocketSubsystemEOSUtils_OnlineServicesEOSGS.cpp b/Engine/Plugins/Online/OnlineServicesEOSGS/Source/Private/Online/SocketSubsystemEOSUtils_OnlineServicesEOSGS.cpp
--- a/Engine/Plugins/Online/OnlineServicesEOSGS/Source/Private/Online/SocketSubsystemEOSUtils_OnlineServicesEOSGS.cpp
+++ b/Engine/Plugins/Online/OnlineServicesEOSGS/Source/Private/Online/SocketSubsystemEOSUtils_OnlineServicesEOSGS.cpp
@@ -46,6 +46,14 @@ EOS_ProductUserId FSocketSubsystemEOSUtils_OnlineServicesEOS::GetLocalUserId()
 	return Result;
 }
 
+bool FSocketSubsystemEOSUtils_OnlineServicesEOS::IsLoggedIn()
+{
+	IAuthPtr Auth = ServicesEOSGS.GetAuthInterface();
+	check(Auth);
+
+	return Auth->IsLoggedIn(FPlatformMisc::GetPlatformUserForUserIndex(0));
+}
+
 FString FSocketSubsystemEOSUtils_OnlineServicesEOS::GetSessionId()
 {
 	FString Result;
diff --git a/Engine/Plugins/Online/OnlineServicesEOSGS/Source/Private/Online/SocketSubsystemEOSUtils_OnlineServicesEOSGS.h b/Engine/Plugins/Online/OnlineServicesEOSGS/Source/Private/Online/SocketSubsystemEOSUtils_OnlineServicesEOSGS.h
--- a/Engine/Plugins/Online/OnlineServicesEOSGS/Source/Private/Online/SocketSubsystemEOSUtils_OnlineServicesEOSGS.h
+++ b/Engine/Plugins/Online/OnlineServicesEOSGS/Source/Private/Online/SocketSubsystemEOSUtils_OnlineServicesEOSGS.h
@@ -19,6 +19,7 @@ public:
 	virtual EOS_ProductUserId GetLocalUserId() override;
 	virtual FString GetSessionId() override;
 	virtual FName GetSubsystemInstanceName() override;
+	virtual bool IsLoggedIn() override;
 
 private:
 	FOnlineServicesEOSGS& ServicesEOSGS;
diff --git a/Engine/Plugins/Online/OnlineSubsystemEOS/Source/OnlineSubsystemEOS/Private/SocketSubsystemEOSUtils_OnlineSubsystemEOS.cpp b/Engine/Plugins/Online/OnlineSubsystemEOS/Source/OnlineSubsystemEOS/Private/SocketSubsystemEOSUtils_OnlineSubsystemEOS.cpp
--- a/Engine/Plugins/Online/OnlineSubsystemEOS/Source/OnlineSubsystemEOS/Private/SocketSubsystemEOSUtils_OnlineSubsystemEOS.cpp
+++ b/Engine/Plugins/Online/OnlineSubsystemEOS/Source/OnlineSubsystemEOS/Private/SocketSubsystemEOSUtils_OnlineSubsystemEOS.cpp
@@ -2,6 +2,7 @@
 
 #include "SocketSubsystemEOSUtils_OnlineSubsystemEOS.h"
 
+#include "OnlineSubsystemTypes.h"
 #include "OnlineSubsystemEOS.h"
 #include "UserManagerEOS.h"
 #include "OnlineSessionEOS.h"
@@ -19,11 +20,19 @@ EOS_ProductUserId FSocketSubsystemEOSUtils_OnlineSubsystemEOS::GetLocalUserId()
 {
 	EOS_ProductUserId Result = nullptr;
 
+	check(SubsystemEOS.UserManager);
 	Result = SubsystemEOS.UserManager->GetLocalProductUserId(SubsystemEOS.UserManager->GetDefaultLocalUser());
 
 	return Result;
 }
 
+bool FSocketSubsystemEOSUtils_OnlineSubsystemEOS::IsLoggedIn()
+{
+	check(SubsystemEOS.UserManager);
+	ELoginStatus::Type LoginStatus = SubsystemEOS.UserManager->GetLoginStatus(SubsystemEOS.UserManager->GetDefaultLocalUser());
+	return LoginStatus == ELoginStatus::LoggedIn; 
+}
+
 FString FSocketSubsystemEOSUtils_OnlineSubsystemEOS::GetSessionId()
 {
 	FString Result;
diff --git a/Engine/Plugins/Online/OnlineSubsystemEOS/Source/OnlineSubsystemEOS/Private/SocketSubsystemEOSUtils_OnlineSubsystemEOS.h b/Engine/Plugins/Online/OnlineSubsystemEOS/Source/OnlineSubsystemEOS/Private/SocketSubsystemEOSUtils_OnlineSubsystemEOS.h
--- a/Engine/Plugins/Online/OnlineSubsystemEOS/Source/OnlineSubsystemEOS/Private/SocketSubsystemEOSUtils_OnlineSubsystemEOS.h
+++ b/Engine/Plugins/Online/OnlineSubsystemEOS/Source/OnlineSubsystemEOS/Private/SocketSubsystemEOSUtils_OnlineSubsystemEOS.h
@@ -15,6 +15,7 @@ public:
 	virtual EOS_ProductUserId GetLocalUserId() override;
 	virtual FString GetSessionId() override;
 	virtual FName GetSubsystemInstanceName() override;
+	virtual bool IsLoggedIn() override;;
 
 private:
 	FOnlineSubsystemEOS& SubsystemEOS;
diff --git a/Engine/Plugins/Online/SocketSubsystemEOS/Source/SocketSubsystemEOS/Private/SocketEOS.cpp b/Engine/Plugins/Online/SocketSubsystemEOS/Source/SocketSubsystemEOS/Private/SocketEOS.cpp
--- a/Engine/Plugins/Online/SocketSubsystemEOS/Source/SocketSubsystemEOS/Private/SocketEOS.cpp
+++ b/Engine/Plugins/Online/SocketSubsystemEOS/Source/SocketSubsystemEOS/Private/SocketEOS.cpp
@@ -413,6 +413,13 @@ bool FSocketEOS::RecvFrom(uint8* Data, int32 BufferSize, int32& BytesRead, FInte
 	check(IsInGameThread() && "p2p does not support multithreading");
 	BytesRead = 0;
 
+	if (!SocketSubsystem.GetSocketSubsystemEOSUtils()->IsLoggedIn())
+	{
+		UE_LOG(LogSocketSubsystemEOS, VeryVerbose, TEXT("Local user is logged out"));
+
+		return false;
+	}
+	
 	if (BufferSize < 0)
 	{
 		UE_LOG(LogSocketSubsystemEOS, Error, TEXT("Unable to receive data, receiving buffer was invalid. BufferSize=[%d]"), BufferSize);
diff --git a/Engine/Plugins/Online/SocketSubsystemEOS/Source/SocketSubsystemEOS/Public/SocketSubsystemEOS.h b/Engine/Plugins/Online/SocketSubsystemEOS/Source/SocketSubsystemEOS/Public/SocketSubsystemEOS.h
--- a/Engine/Plugins/Online/SocketSubsystemEOS/Source/SocketSubsystemEOS/Public/SocketSubsystemEOS.h
+++ b/Engine/Plugins/Online/SocketSubsystemEOS/Source/SocketSubsystemEOS/Public/SocketSubsystemEOS.h
@@ -91,6 +91,13 @@ public:
 	 */
 	FSocketSubsystemEOS* GetSocketSubsystemForWorld(UWorld* InWorld);
 
+	/**
+	 * Returns the utils for this SocketSubsystemEOS
+	 *
+	 * @return The ISocketSubsystemEOSUtilsPtr for this SocketSubsystemEOS
+	 */
+	ISocketSubsystemEOSUtilsPtr GetSocketSubsystemEOSUtils() { return Utils; }
+
 private:
 	/** Removes the FSocketSubsystemEOS instance from SocketSubsystemEOSInstances and all related entries from SocketSubsystemEOSPerWorldMap */
 	void RemoveFromStaticContainers();
diff --git a/Engine/Plugins/Online/SocketSubsystemEOS/Source/SocketSubsystemEOS/Public/SocketSubsystemEOSUtils.h b/Engine/Plugins/Online/SocketSubsystemEOS/Source/SocketSubsystemEOS/Public/SocketSubsystemEOSUtils.h
--- a/Engine/Plugins/Online/SocketSubsystemEOS/Source/SocketSubsystemEOS/Public/SocketSubsystemEOSUtils.h
+++ b/Engine/Plugins/Online/SocketSubsystemEOS/Source/SocketSubsystemEOS/Public/SocketSubsystemEOSUtils.h
@@ -23,5 +23,7 @@ public:
 	virtual FString GetSessionId() = 0;
 
 	virtual FName GetSubsystemInstanceName() = 0;
+
+	virtual bool IsLoggedIn() = 0;
 };
 typedef TSharedPtr<ISocketSubsystemEOSUtils, ESPMode::ThreadSafe> ISocketSubsystemEOSUtilsPtr;
