- Fix logging Error when PlatformOSS is not configured
- Leaving weird desktop platform fallback on Steam which absolutely shouldn't be required but I want to avoid breaking whatever weird use case might be covered by it.

#rnx
#DisabledValidations lowLevelTestsOnlineTests, lowLevelTestsOSMCPTests

diff --git a/Engine/Plugins/Online/OnlineSubsystemEOS/Source/OnlineSubsystemEOS/Private/UserManagerEOS.cpp b/Engine/Plugins/Online/OnlineSubsystemEOS/Source/OnlineSubsystemEOS/Private/UserManagerEOS.cpp
--- a/Engine/Plugins/Online/OnlineSubsystemEOS/Source/OnlineSubsystemEOS/Private/UserManagerEOS.cpp
+++ b/Engine/Plugins/Online/OnlineSubsystemEOS/Source/OnlineSubsystemEOS/Private/UserManagerEOS.cpp
@@ -227,17 +227,12 @@ bool ToEOS_ELoginCredentialType(const FString& InTypeStr, EOS_ELoginCredentialTy
 IOnlineSubsystem* GetPlatformOSS()
 {
 	IOnlineSubsystem* PlatformOSS = IOnlineSubsystem::GetByPlatform();
-	if (PlatformOSS == nullptr)
-#if !PLATFORM_DESKTOP
-	{
-		UE_LOG_ONLINE(Error, TEXT("GetPlatformOSS() failed due to no platform OSS being configured"));
-	}
-#else
+	if (!PlatformOSS && PLATFORM_DESKTOP)
 	{
 		// Attempt to load Steam before treating it as an error
+		// Note: Vestigial/nonsensical. If we wanted Steam we could configure NativePlatformService=Steam. Leaving it here to not break some unknown bad config setup.
 		PlatformOSS = IOnlineSubsystem::Get(STEAM_SUBSYSTEM);
 	}
-#endif
 	return PlatformOSS;
 }
 
