- Add ability to configure OnlineSubsystemEOS to use the "named platform config" mechanism

diff --git a/Engine/Plugins/Online/OnlineSubsystemEOS/Source/OnlineSubsystemEOS/Private/EOSSettings.cpp b/Engine/Plugins/Online/OnlineSubsystemEOS/Source/OnlineSubsystemEOS/Private/EOSSettings.cpp
--- a/Engine/Plugins/Online/OnlineSubsystemEOS/Source/OnlineSubsystemEOS/Private/EOSSettings.cpp
+++ b/Engine/Plugins/Online/OnlineSubsystemEOS/Source/OnlineSubsystemEOS/Private/EOSSettings.cpp
@@ -179,6 +179,8 @@ const FEOSSettings& UEOSSettings::ManualGetSettings()
 		GConfig->GetString(INI_SECTION, TEXT("CacheDir"), GCachedSettings->CacheDir, GEngineIni);
 		GConfig->GetString(INI_SECTION, TEXT("DefaultArtifactName"), GCachedSettings->DefaultArtifactName, GEngineIni);
 		GConfig->GetString(INI_SECTION, TEXT("SteamTokenType"), GCachedSettings->SteamTokenType, GEngineIni);
+		GConfig->GetString(INI_SECTION, TEXT("NintendoTokenType"), GCachedSettings->NintendoTokenType, GEngineIni);
+		GConfig->GetString(INI_SECTION, TEXT("PlatformConfigName"), GCachedSettings->PlatformConfigName, GEngineIni);
 		GCachedSettings->RTCBackgroundMode = EOS_ERTCBackgroundMode::EOS_RTCBM_KeepRoomsAlive;
 		FString RTCBackgroundModeStr;
 		GConfig->GetString(INI_SECTION, TEXT("RTCBackgroundMode"), RTCBackgroundModeStr, GEngineIni);
@@ -195,6 +197,7 @@ const FEOSSettings& UEOSSettings::ManualGetSettings()
 		GConfig->GetBool(INI_SECTION, TEXT("bUseEAS"), GCachedSettings->bUseEAS, GEngineIni);
 		GConfig->GetBool(INI_SECTION, TEXT("bUseEOSConnect"), GCachedSettings->bUseEOSConnect, GEngineIni);
 		GConfig->GetBool(INI_SECTION, TEXT("bUseEOSRTC"), GCachedSettings->bUseEOSRTC, GEngineIni);
+		GConfig->GetBool(INI_SECTION, TEXT("bUseNamedPlatformConfig"), GCachedSettings->bUseNamedPlatformConfig, GEngineIni);
 PRAGMA_DISABLE_DEPRECATION_WARNINGS
 		GConfig->GetBool(INI_SECTION, TEXT("bUseNewLoginFlow"), GCachedSettings->bUseNewLoginFlow, GEngineIni);
 PRAGMA_ENABLE_DEPRECATION_WARNINGS
@@ -213,6 +216,8 @@ FEOSSettings UEOSSettings::ToNative() const
 	Native.CacheDir = CacheDir;
 	Native.DefaultArtifactName = DefaultArtifactName;
 	Native.SteamTokenType = SteamTokenType;
+	Native.NintendoTokenType = NintendoTokenType;
+	Native.PlatformConfigName = PlatformConfigName;
 	Native.RTCBackgroundMode = EOS_ERTCBackgroundMode::EOS_RTCBM_KeepRoomsAlive;
 	if (!RTCBackgroundMode.IsEmpty())
 	{
@@ -227,6 +232,7 @@ FEOSSettings UEOSSettings::ToNative() const
 	Native.bUseEAS = bUseEAS;
 	Native.bUseEOSConnect = bUseEOSConnect;
 	Native.bUseEOSRTC = bUseEOSRTC;
+	Native.bUseNamedPlatformConfig = bUseNamedPlatformConfig;
 PRAGMA_DISABLE_DEPRECATION_WARNINGS
 	Native.bUseNewLoginFlow = bUseNewLoginFlow;
 PRAGMA_ENABLE_DEPRECATION_WARNINGS
diff --git a/Engine/Plugins/Online/OnlineSubsystemEOS/Source/OnlineSubsystemEOS/Private/OnlineSubsystemEOS.cpp b/Engine/Plugins/Online/OnlineSubsystemEOS/Source/OnlineSubsystemEOS/Private/OnlineSubsystemEOS.cpp
--- a/Engine/Plugins/Online/OnlineSubsystemEOS/Source/OnlineSubsystemEOS/Private/OnlineSubsystemEOS.cpp
+++ b/Engine/Plugins/Online/OnlineSubsystemEOS/Source/OnlineSubsystemEOS/Private/OnlineSubsystemEOS.cpp
@@ -175,6 +175,16 @@ void FOnlineSubsystemEOS::ModuleShutdown()
 /** Common method for creating the EOS platform */
 bool FOnlineSubsystemEOS::PlatformCreate()
 {
+	const FEOSSettings& EOSSettings = UEOSSettings::GetSettings();
+
+	if (EOSSettings.bUseNamedPlatformConfig)
+	{
+		EOSPlatformHandle = EOSSDKManager->CreatePlatform(EOSSettings.PlatformConfigName);
+		UE_CLOG_ONLINE(!EOSPlatformHandle.IsValid(), Error, TEXT("FOnlineSubsystemEOS::PlatformCreate() PlatformConfigName=%s failed"),
+			*EOSSettings.PlatformConfigName);
+		return EOSPlatformHandle.IsValid();
+	}
+
 	FEOSArtifactSettings ArtifactSettings;
 	if (!UEOSSettings::GetSelectedArtifactSettings(ArtifactSettings))
 	{
@@ -203,7 +213,6 @@ bool FOnlineSubsystemEOS::PlatformCreate()
 	PlatformOptions.Reserved = nullptr;
 	PlatformOptions.TaskNetworkTimeoutSeconds = nullptr;
 
-	FEOSSettings EOSSettings = UEOSSettings::GetSettings();
 	uint64 OverlayFlags = 0;
 	if (!EOSSettings.bEnableOverlay)
 	{
diff --git a/Engine/Plugins/Online/OnlineSubsystemEOS/Source/OnlineSubsystemEOS/Public/EOSSettings.h b/Engine/Plugins/Online/OnlineSubsystemEOS/Source/OnlineSubsystemEOS/Public/EOSSettings.h
--- a/Engine/Plugins/Online/OnlineSubsystemEOS/Source/OnlineSubsystemEOS/Public/EOSSettings.h
+++ b/Engine/Plugins/Online/OnlineSubsystemEOS/Source/OnlineSubsystemEOS/Public/EOSSettings.h
@@ -80,6 +80,7 @@ PRAGMA_ENABLE_DEPRECATION_WARNINGS
 	FString DefaultArtifactName;
 	FString SteamTokenType;
 	FString NintendoTokenType;
+	FString PlatformConfigName;
 	EOS_ERTCBackgroundMode RTCBackgroundMode;
 	int32 TickBudgetInMilliseconds;
 	int32 TitleStorageReadChunkLength;
@@ -90,6 +91,7 @@ PRAGMA_ENABLE_DEPRECATION_WARNINGS
 	bool bUseEAS;
 	bool bUseEOSConnect;
 	bool bUseEOSRTC;
+	bool bUseNamedPlatformConfig;
 	UE_DEPRECATED(5.7, "bUseNewLoginFlow is deprecated, the legacy login flow has been removed from UserManagerEOS.")
 	bool bUseNewLoginFlow;
 	TArray<FEOSArtifactSettings> Artifacts;
@@ -116,6 +118,10 @@ public:
 	UPROPERTY(Config, EditAnywhere, BlueprintReadOnly, Category="EOS Settings")
 	FString DefaultArtifactName;
 
+	/** When bUseNamedPlatformConfig is true, specifies the platform config name to use */
+	UPROPERTY(Config, EditAnywhere, BlueprintReadOnly, Category="EOS Settings")
+	FString PlatformConfigName;
+
 	/** The preferred background mode to be used by RTC services */
 	UPROPERTY(Config, EditAnywhere, BlueprintReadOnly, Category="EOS Settings")
 	FString RTCBackgroundMode;
@@ -168,6 +174,14 @@ public:
 	UPROPERTY(Config, EditAnywhere, BlueprintReadOnly, Category="EOS Settings", DisplayName="Enable Real-Time chat")
 	bool bUseEOSRTC = true;
 
+	/**
+	 * Set to true to use the IEOSSDKManager "named platform config" mechanism, rather than
+	 * loading a config from OSSEOS ini. Specify the platform config name to use with the
+	 * "PlatformConfigName" setting.
+	 */
+	UPROPERTY(Config, EditAnywhere, BlueprintReadOnly, Category="EOS Settings")
+	bool bUseNamedPlatformConfig = false;
+
 	/** Set to true to use new EOS login flow */
 	UE_DEPRECATED(5.7, "bUseNewLoginFlow is deprecated, the legacy login flow has been removed from UserManagerEOS.")
 	UPROPERTY(Config, EditAnywhere, BlueprintReadOnly, Category="EOS Login Settings", DisplayName="Use new Login flow, which doesn't rely on EOSPlus")
