- Fix error in OSSEOS::Init when using new platform config mechanism

diff --git a/Engine/Plugins/Online/OnlineSubsystemEOS/Source/OnlineSubsystemEOS/Private/OnlineSubsystemEOS.cpp b/Engine/Plugins/Online/OnlineSubsystemEOS/Source/OnlineSubsystemEOS/Private/OnlineSubsystemEOS.cpp
--- a/Engine/Plugins/Online/OnlineSubsystemEOS/Source/OnlineSubsystemEOS/Private/OnlineSubsystemEOS.cpp
+++ b/Engine/Plugins/Online/OnlineSubsystemEOS/Source/OnlineSubsystemEOS/Private/OnlineSubsystemEOS.cpp
@@ -381,16 +381,41 @@ bool FOnlineSubsystemEOS::Init()
 	}
 
 	// We set the product id
-	FEOSArtifactSettings ArtifactSettings;
-	if (UEOSSettings::GetSelectedArtifactSettings(ArtifactSettings))
+	FString EncryptionKey, ArtifactName;
+
+	const FEOSSettings& EOSSettings = UEOSSettings::GetSettings();
+	if (EOSSettings.bUseNamedPlatformConfig)
 	{
-		ProductId = ArtifactSettings.ProductId;
+		if (const FEOSSDKPlatformConfig* PlatformConfig = EOSSDKManager->GetPlatformConfig(EOSSettings.PlatformConfigName))
+		{
+			ProductId = PlatformConfig->ProductId;
+			EncryptionKey = PlatformConfig->EncryptionKey;
+			ArtifactName = PlatformConfig->Name;
+		}
+		else
+		{
+			// it souldn't be possible to get here if PlatformCreate succeeded
+			checkNoEntry();
+			UE_LOG_ONLINE(Error, TEXT("[FOnlineSubsystemEOS::Init] GetPlatformConfig failed"));
+			return false;
+		}
 	}
 	else
 	{
-		// This really should not be possible, if we made it past PlatformCreate.
-		checkNoEntry();
-		UE_LOG_ONLINE(Warning, TEXT("[FOnlineSubsystemEOS::Init] GetSelectedArtifactSettings failed, ProductIdAnsi not set."));
+		FEOSArtifactSettings ArtifactSettings;
+		if (UEOSSettings::GetSelectedArtifactSettings(ArtifactSettings))
+		{
+			ProductId = ArtifactSettings.ProductId;
+			EncryptionKey = ArtifactSettings.EncryptionKey;
+			ArtifactName = ArtifactSettings.ArtifactName;
+		}
+		else
+		{
+			// it souldn't be possible to get here if PlatformCreate succeeded
+			checkNoEntry();
+			UE_LOG_ONLINE(Error, TEXT("[FOnlineSubsystemEOS::Init] GetSelectedArtifactSettings failed"));
+			return false;
+		}
 	}
 
 	TitleStorageHandle = EOS_Platform_GetTitleStorageInterface(*EOSPlatformHandle);
@@ -398,9 +423,9 @@ bool FOnlineSubsystemEOS::Init()
 	{
 		TitleFileInterfacePtr = MakeShared<FOnlineTitleFileEOS>(this);
 	}
-	else if (ArtifactSettings.EncryptionKey.IsEmpty())
+	else if (EncryptionKey.IsEmpty())
 	{
-		UE_LOG_ONLINE(Verbose, TEXT("FOnlineSubsystemEOS: ArtifactName=[%s] EncryptionKey unset, TitleFile interface unavailable."), *ArtifactSettings.ArtifactName);
+		UE_LOG_ONLINE(Verbose, TEXT("FOnlineSubsystemEOS: ArtifactName=[%s] EncryptionKey unset, TitleFile interface unavailable."), *ArtifactName);
 	}
 	else
 	{
@@ -413,9 +438,9 @@ bool FOnlineSubsystemEOS::Init()
 	{
 		UserCloudInterfacePtr = MakeShared<FOnlineUserCloudEOS>(this);
 	}
-	else if (ArtifactSettings.EncryptionKey.IsEmpty())
+	else if (EncryptionKey.IsEmpty())
 	{
-		UE_LOG_ONLINE(Verbose, TEXT("FOnlineSubsystemEOS: ArtifactName=[%s] EncryptionKey unset, UserCloud interface unavailable."), *ArtifactSettings.ArtifactName);
+		UE_LOG_ONLINE(Verbose, TEXT("FOnlineSubsystemEOS: ArtifactName=[%s] EncryptionKey unset, UserCloud interface unavailable."), *ArtifactName);
 	}
 	else
 	{
