- Map DefaultInstanceName to NAME_None

diff --git a/Engine/Plugins/Online/OnlineSubsystemEOS/Source/OnlineSubsystemEOS/Private/OnlineSubsystemEOS.cpp b/Engine/Plugins/Online/OnlineSubsystemEOS/Source/OnlineSubsystemEOS/Private/OnlineSubsystemEOS.cpp
--- a/Engine/Plugins/Online/OnlineSubsystemEOS/Source/OnlineSubsystemEOS/Private/OnlineSubsystemEOS.cpp
+++ b/Engine/Plugins/Online/OnlineSubsystemEOS/Source/OnlineSubsystemEOS/Private/OnlineSubsystemEOS.cpp
@@ -179,7 +179,7 @@ bool FOnlineSubsystemEOS::PlatformCreate()
 
 	if (EOSSettings.bUseNamedPlatformConfig)
 	{
-		EOSPlatformHandle = EOSSDKManager->CreatePlatform(EOSSettings.PlatformConfigName);
+		EOSPlatformHandle = EOSSDKManager->CreatePlatform(EOSSettings.PlatformConfigName, InstanceName == FOnlineSubsystemImpl::DefaultInstanceName ? NAME_None : InstanceName);
 		UE_CLOG_ONLINE(!EOSPlatformHandle.IsValid(), Error, TEXT("FOnlineSubsystemEOS::PlatformCreate() PlatformConfigName=%s failed"),
 			*EOSSettings.PlatformConfigName);
 		return EOSPlatformHandle.IsValid();
