Modify method of retrieving SandboxId in FOnlineStoreEOS::QueryReceipts

diff --git a/Engine/Plugins/Online/OnlineSubsystemEOS/Source/OnlineSubsystemEOS/Private/OnlineStoreEOS.cpp b/Engine/Plugins/Online/OnlineSubsystemEOS/Source/OnlineSubsystemEOS/Private/OnlineStoreEOS.cpp
--- a/Engine/Plugins/Online/OnlineSubsystemEOS/Source/OnlineSubsystemEOS/Private/OnlineStoreEOS.cpp
+++ b/Engine/Plugins/Online/OnlineSubsystemEOS/Source/OnlineSubsystemEOS/Private/OnlineStoreEOS.cpp
@@ -1,6 +1,8 @@
 // Copyright Epic Games, Inc. All Rights Reserved.
 
 #include "OnlineStoreEOS.h"
+#include "EOSSettings.h"
+#include "IEOSSDKManager.h"
 
 #if WITH_EOS_SDK
 
@@ -373,22 +375,47 @@ void FOnlineStoreEOS::QueryReceipts(const FUniqueNetId& UserId, bool bRestoreRec
 
 		// After QueryEntitlements is done, we call QueryOwnership to query ownership granted by other services like season passes, any benefits of which would not have been returned by QueryEntitlements
 
+		// Retrieve the SandboxId from config
+		FString SandboxIdString;
+		const FEOSSettings& EOSSettings = UEOSSettings::GetSettings();
+		if (EOSSettings.bUseNamedPlatformConfig)
+		{
+			const FEOSSDKPlatformConfig* const PlatformConfig = IEOSSDKManager::Get()->GetPlatformConfig(EOSSettings.PlatformConfigName);
+			if (!PlatformConfig)
+			{
+				UE_LOG_ONLINE(Warning, TEXT("QueryReceipts - GetPlatformConfig failed - EOS_Ecom_QueryOwnershipBySandboxIds not called"));
+				OnComplete.ExecuteIfBound(ONLINE_ERROR(EOnlineErrorResult::Unknown));
+				return;
+			}
+			SandboxIdString = PlatformConfig->SandboxId; 
+		}
+		else
+		{
+			FEOSArtifactSettings ArtifactSettings;
+			if (!UEOSSettings::GetSelectedArtifactSettings(ArtifactSettings))
+			{
+				UE_LOG_ONLINE(Warning, TEXT("QueryReceipts - GetSelectedArtifactSettings failed - EOS_Ecom_QueryOwnershipBySandboxIds not called"));
+				OnComplete.ExecuteIfBound(ONLINE_ERROR(EOnlineErrorResult::Unknown));
+				return;
+			}
+			SandboxIdString = ArtifactSettings.SandboxId;
+		}
+
+		if (SandboxIdString.IsEmpty())
+		{
+			UE_LOG_ONLINE(Warning, TEXT("QueryReceipts - Empty SandboxId - EOS_Ecom_QueryOwnershipBySandboxIds not called"));
+			OnComplete.ExecuteIfBound(ONLINE_ERROR(EOnlineErrorResult::Unknown));
+			return;
+		}
+
 		EOS_Ecom_QueryOwnershipBySandboxIdsOptions Options = {};
 		Options.ApiVersion = 1;
 		UE_EOS_CHECK_API_MISMATCH(EOS_ECOM_QUERYOWNERSHIPBYSANDBOXIDSOPTIONS_API_LATEST, 1);
 		Options.LocalUserId = Data->LocalUserId;
 		Options.SandboxIdsCount = 1;
-
-		if (IEOSSDKManager* Manager = IEOSSDKManager::Get())
-		{
-			const FString& PlatformConfigName = EOSSubsystem->GetEOSPlatformHandle()->GetConfigName();
-			if (const FEOSSDKPlatformConfig* Config = Manager->GetPlatformConfig(PlatformConfigName))
-			{
-				const FTCHARToUTF8& Converter = FTCHARToUTF8(*Config->SandboxId);
-				const char* SandboxId = Converter.Get();
-				Options.SandboxIds = &SandboxId;
-			}
-		}		
+		const auto SandboxIdUtf8 = StringCast<UTF8CHAR>(*SandboxIdString);
+		const char* SandboxId = (const char*)SandboxIdUtf8.Get();
+		Options.SandboxIds = &SandboxId;
 
 		FQueryOwnershipBySandboxIdsCallback* CallbackObj = new FQueryOwnershipBySandboxIdsCallback(FOnlineStoreEOSWeakPtr(AsShared()));
 		CallbackObj->CallbackLambda = [this, OnComplete = FOnQueryReceiptsComplete(OnComplete)](const EOS_Ecom_QueryOwnershipBySandboxIdsCallbackInfo* Data)
