Missing nullptr check in QueryReceipts - Modified code to get from EOSSubsystem

diff --git a/Engine/Plugins/Online/OnlineSubsystemEOS/Source/OnlineSubsystemEOS/Private/OnlineStoreEOS.cpp b/Engine/Plugins/Online/OnlineSubsystemEOS/Source/OnlineSubsystemEOS/Private/OnlineStoreEOS.cpp
--- a/Engine/Plugins/Online/OnlineSubsystemEOS/Source/OnlineSubsystemEOS/Private/OnlineStoreEOS.cpp
+++ b/Engine/Plugins/Online/OnlineSubsystemEOS/Source/OnlineSubsystemEOS/Private/OnlineStoreEOS.cpp
@@ -380,7 +380,7 @@ void FOnlineStoreEOS::QueryReceipts(const FUniqueNetId& UserId, bool bRestoreRec
 		const FEOSSettings& EOSSettings = UEOSSettings::GetSettings();
 		if (EOSSettings.bUseNamedPlatformConfig)
 		{
-			const FEOSSDKPlatformConfig* const PlatformConfig = IEOSSDKManager::Get()->GetPlatformConfig(EOSSettings.PlatformConfigName);
+			const FEOSSDKPlatformConfig* const PlatformConfig = EOSSubsystem->EOSSDKManager->GetPlatformConfig(EOSSettings.PlatformConfigName);
 			if (!PlatformConfig)
 			{
 				UE_LOG_ONLINE(Warning, TEXT("QueryReceipts - GetPlatformConfig failed - EOS_Ecom_QueryOwnershipBySandboxIds not called"));
