Add condition to use XSTS token in OnlineSubsystemEOS when target is WinGDK

diff --git a/Engine/Plugins/Online/OnlineSubsystemEOS/Source/OnlineSubsystemEOS/OnlineSubsystemEOS.Build.cs b/Engine/Plugins/Online/OnlineSubsystemEOS/Source/OnlineSubsystemEOS/OnlineSubsystemEOS.Build.cs
--- a/Engine/Plugins/Online/OnlineSubsystemEOS/Source/OnlineSubsystemEOS/OnlineSubsystemEOS.Build.cs
+++ b/Engine/Plugins/Online/OnlineSubsystemEOS/Source/OnlineSubsystemEOS/OnlineSubsystemEOS.Build.cs
@@ -46,7 +46,7 @@ public class OnlineSubsystemEOS : ModuleRules
 		PrivateDefinitions.Add("EOS_AUTH_TOKEN_SAVEGAME_STORAGE=" + (bAuthTokenSavegameStorage ? "1" : "0"));
 	}
 
-	protected virtual bool bUseXblXstsToken { get { return false; } }
+	protected virtual bool bUseXblXstsToken { get { return Target.Platform == UnrealTargetPlatform.WinGDK; } }
 	protected virtual bool bUsePsnIdToken { get { return false; } }
 	protected virtual bool bAddUserLoginInfo { get { return false; } }
 	protected virtual bool bAuthTokenSavegameStorage { get { return false; } }
