[EOS] Fixed issue where connection to EOS fails when auth is already pending.

When calling `FUserManagerEOS::Login`, the server may respond with `EOS_Auth_PinGrantPending` during the authentication process. This was processed as an error, but it only means that the auth request was already made once, no need to make another request. If this happens, we now only log the error, but we don't cancel the login process anymore.

diff --git a/Engine/Plugins/Online/OnlineSubsystemEOS/Source/OnlineSubsystemEOS/Private/UserManagerEOS.cpp b/Engine/Plugins/Online/OnlineSubsystemEOS/Source/OnlineSubsystemEOS/Private/UserManagerEOS.cpp
--- a/Engine/Plugins/Online/OnlineSubsystemEOS/Source/OnlineSubsystemEOS/Private/UserManagerEOS.cpp
+++ b/Engine/Plugins/Online/OnlineSubsystemEOS/Source/OnlineSubsystemEOS/Private/UserManagerEOS.cpp
@@ -788,6 +788,11 @@ void FUserManagerEOS::OnEOSAuthLoginComplete(int32 LocalUserNum, const EOS_ELogi
 		Credentials.Type = TEXT("accountportal");
 		CallEOSAuthLogin(LocalUserNum, Credentials, true);
 	}
+	else if (Data->ResultCode == EOS_EResult::EOS_Auth_PinGrantPending)
+	{
+		FString PendingMessage = FString::Printf(TEXT("Login(%d) Auth attempt already pending (%s)"), LocalUserNum, *LexToString(Data->ResultCode));
+		UE_LOG_ONLINE(Warning, TEXT("%s"), *PendingMessage);
+	}
 	else
 	{
 		auto TriggerLoginFailure = [this, LocalUserNum, LoginResultCode = Data->ResultCode]()
