- Fix EOS_ByteArray_ToString being called with invalid parameters

[RN] Fixed failure to log in to OnlineSubsystemEOS with Steam credentials when bUseEAS=false

diff --git a/Engine/Plugins/Online/OnlineSubsystemEOS/Source/OnlineSubsystemEOS/Private/UserManagerEOS.cpp b/Engine/Plugins/Online/OnlineSubsystemEOS/Source/OnlineSubsystemEOS/Private/UserManagerEOS.cpp
--- a/Engine/Plugins/Online/OnlineSubsystemEOS/Source/OnlineSubsystemEOS/Private/UserManagerEOS.cpp
+++ b/Engine/Plugins/Online/OnlineSubsystemEOS/Source/OnlineSubsystemEOS/Private/UserManagerEOS.cpp
@@ -904,11 +904,7 @@ FString ConvertExternalAuthTokenToFString(const FExternalAuthToken& AuthToken)
 
 	if (AuthToken.HasTokenData())
 	{
-		char* TokenAnsi = nullptr;
-		uint32_t InOutBufferLength = EOS_MAX_TOKEN_SIZE;
-		EOS_ByteArray_ToString(AuthToken.TokenData.GetData(), AuthToken.TokenData.Num(), TokenAnsi, &InOutBufferLength);
-
-		Result = FString(TokenAnsi);
+		Result = ToHexString(AuthToken.TokenData);
 	}
 	else if (AuthToken.HasTokenString())
 	{
