- Fix CopySessionInfo capturing a reference to an array element that is then relocated

[RN] Fix crash when a session search yields 5+ results

diff --git a/Engine/Plugins/Online/OnlineSubsystemEOS/Source/OnlineSubsystemEOS/Private/OnlineSessionEOS.cpp b/Engine/Plugins/Online/OnlineSubsystemEOS/Source/OnlineSubsystemEOS/Private/OnlineSessionEOS.cpp
--- a/Engine/Plugins/Online/OnlineSubsystemEOS/Source/OnlineSubsystemEOS/Private/OnlineSessionEOS.cpp
+++ b/Engine/Plugins/Online/OnlineSubsystemEOS/Source/OnlineSubsystemEOS/Private/OnlineSessionEOS.cpp
@@ -2386,51 +2386,6 @@ void FOnlineSessionEOS::AddLobbySearchAttribute(EOS_HLobbySearch LobbySearchHand
 	}
 }
 
-void FOnlineSessionEOS::CopySearchResult(const FSessionDetailsEOS& SessionHandle, EOS_SessionDetails_Info* SessionInfo, FOnlineSession& OutSession, FOnCopySessionDataCompleteCallback&& Callback)
-{
-	OutSession.NumOpenPrivateConnections = SessionInfo->NumOpenPublicConnections;
-	OutSession.SessionSettings.NumPrivateConnections = SessionInfo->Settings->NumPublicConnections;
-	OutSession.SessionSettings.bAllowJoinInProgress = SessionInfo->Settings->bAllowJoinInProgress == EOS_TRUE;
-	OutSession.SessionSettings.bAllowInvites = SessionInfo->Settings->bInvitesAllowed == EOS_TRUE;
-
-	switch (SessionInfo->Settings->PermissionLevel)
-	{
-		case EOS_EOnlineSessionPermissionLevel::EOS_OSPF_InviteOnly:
-		{
-			OutSession.SessionSettings.bAllowJoinViaPresence = false;
-			break;
-		}
-		case EOS_EOnlineSessionPermissionLevel::EOS_OSPF_JoinViaPresence:
-		case EOS_EOnlineSessionPermissionLevel::EOS_OSPF_PublicAdvertised:
-		{
-			OutSession.SessionSettings.bAllowJoinViaPresence = true;
-			break;
-		}
-	}
-
-	CopyAttributes(SessionHandle, OutSession);
-
-	if (SessionInfo->OwnerServerClientId)
-	{
-		OutSession.OwningUserId = FUniqueNetIdEOS::DedicatedServerId();
-
-		// ResolveUniqueNetIds is an asynchronous operation, so in the cases where it's not called, we'll delay the execution of this callback to match the flow
-		EOSSubsystem->ExecuteNextTick([Callback = MoveTemp(Callback)]()
-			{
-				Callback(true);
-			});
-	}
-	else
-	{
-		EOSSubsystem->UserManager->ResolveUniqueNetId(EOSSubsystem->UserManager->GetDefaultLocalUser(), SessionInfo->OwnerUserId, [this, &OutSession, Callback = MoveTemp(Callback)](FUniqueNetIdEOSRef ResolvedUniqueNetId, const FOnlineError& Error)
-			{
-				OutSession.OwningUserId = ResolvedUniqueNetId;
-				OutSession.OwningUserName = EOSSubsystem->UserManager->GetPlayerNickname(*ResolvedUniqueNetId);
-				Callback(true);
-			});
-	}
-}
-
 void FOnlineSessionEOS::CopyAttributes(const FSessionDetailsEOS& SessionHandle, FOnlineSession& OutSession)
 {
 	EOS_SessionDetails_GetSessionAttributeCountOptions CountOptions = { };
@@ -2525,14 +2480,58 @@ void FOnlineSessionEOS::AddSearchResult(const TSharedRef<FSessionDetailsEOS>& Se
 	EOS_EResult CopyResult = EOS_SessionDetails_CopyInfo(SessionHandle->SessionDetailsHandle, &CopyOptions, &EosSessionDetailsInfo);
 	if (CopyResult == EOS_EResult::EOS_Success)
 	{
-		int32 Position = SearchSettings->SearchResults.AddZeroed();
+		const int32 Position = SearchSettings->SearchResults.AddZeroed();
 		FOnlineSessionSearchResult& SearchResult = SearchSettings->SearchResults[Position];
 
 		// This will set the host address and port
 		TSharedPtr<FOnlineSessionInfoEOS> OnlineSessionInfo = MakeShared<FOnlineSessionInfoEOS>(FOnlineSessionInfoEOS::Create(FUniqueNetIdEOSSession::Create(EosSessionDetailsInfo->SessionId), SessionHandle));
 		SearchResult.Session.SessionInfo = OnlineSessionInfo;
 
-		CopySearchResult(*SessionHandle, EosSessionDetailsInfo, SearchResult.Session, MoveTemp(Callback));
+		SearchResult.Session.NumOpenPrivateConnections = EosSessionDetailsInfo->NumOpenPublicConnections;
+		SearchResult.Session.SessionSettings.NumPrivateConnections = EosSessionDetailsInfo->Settings->NumPublicConnections;
+		SearchResult.Session.SessionSettings.bAllowJoinInProgress = EosSessionDetailsInfo->Settings->bAllowJoinInProgress == EOS_TRUE;
+		SearchResult.Session.SessionSettings.bAllowInvites = EosSessionDetailsInfo->Settings->bInvitesAllowed == EOS_TRUE;
+
+		switch (EosSessionDetailsInfo->Settings->PermissionLevel)
+		{
+			case EOS_EOnlineSessionPermissionLevel::EOS_OSPF_InviteOnly:
+			{
+				SearchResult.Session.SessionSettings.bAllowJoinViaPresence = false;
+				break;
+			}
+			case EOS_EOnlineSessionPermissionLevel::EOS_OSPF_JoinViaPresence:
+			case EOS_EOnlineSessionPermissionLevel::EOS_OSPF_PublicAdvertised:
+			{
+				SearchResult.Session.SessionSettings.bAllowJoinViaPresence = true;
+				break;
+			}
+		}
+
+		CopyAttributes(*SessionHandle, SearchResult.Session);
+
+		if (EosSessionDetailsInfo->OwnerServerClientId)
+		{
+			SearchResult.Session.OwningUserId = FUniqueNetIdEOS::DedicatedServerId();
+
+			// ResolveUniqueNetIds is an asynchronous operation, so in the cases where it's not called, we'll delay the execution of this callback to match the flow
+			EOSSubsystem->ExecuteNextTick([Callback = MoveTemp(Callback)]()
+				{
+					Callback(true);
+				});
+		}
+		else
+		{
+			EOSSubsystem->UserManager->ResolveUniqueNetId(EOSSubsystem->UserManager->GetDefaultLocalUser(), EosSessionDetailsInfo->OwnerUserId, [this, WeakSearchSettings = SearchSettings.ToWeakPtr(), Position, Callback = MoveTemp(Callback)](FUniqueNetIdEOSRef ResolvedUniqueNetId, const FOnlineError& Error)
+			{
+				if (TSharedPtr<FOnlineSessionSearch> StrongSearchSettings = WeakSearchSettings.Pin())
+				{
+					FOnlineSession& OutSession = StrongSearchSettings->SearchResults[Position].Session;
+					OutSession.OwningUserId = ResolvedUniqueNetId;
+					OutSession.OwningUserName = EOSSubsystem->UserManager->GetPlayerNickname(*ResolvedUniqueNetId);
+				}
+				Callback(true);
+			});
+		}
 
 		// CopySearchResult above will populate the settings so we can now read the port and construct the HostAddress.
 		OnlineSessionInfo->HostAddr = GetHostAddress(UTF8_TO_TCHAR(EosSessionDetailsInfo->HostAddress), SearchResult.Session.SessionSettings);
diff --git a/Engine/Plugins/Online/OnlineSubsystemEOS/Source/OnlineSubsystemEOS/Private/OnlineSessionEOS.h b/Engine/Plugins/Online/OnlineSubsystemEOS/Source/OnlineSubsystemEOS/Private/OnlineSessionEOS.h
--- a/Engine/Plugins/Online/OnlineSubsystemEOS/Source/OnlineSubsystemEOS/Private/OnlineSessionEOS.h
+++ b/Engine/Plugins/Online/OnlineSubsystemEOS/Source/OnlineSubsystemEOS/Private/OnlineSessionEOS.h
@@ -302,7 +302,6 @@ private:
 	typedef TFunction<void(bool bWasSuccessful)> FOnCopySessionDataCompleteCallback;
 	void AddSearchResult(const TSharedRef<FSessionDetailsEOS>& SessionHandle, const TSharedRef<FOnlineSessionSearch>& SearchSettings, FOnCopySessionDataCompleteCallback&& Callback);
 	void AddSearchAttribute(EOS_HSessionSearch SearchHandle, const EOS_Sessions_AttributeData* Attribute, EOS_EOnlineComparisonOp ComparisonOp);
-	void CopySearchResult(const FSessionDetailsEOS& SessionHandle, EOS_SessionDetails_Info* SessionInfo, FOnlineSession& SessionSettings, FOnCopySessionDataCompleteCallback&& Callback);
 	void CopyAttributes(const FSessionDetailsEOS& SessionHandle, FOnlineSession& OutSession);
 
 	void SetPermissionLevel(EOS_HSessionModification SessionModHandle, FNamedOnlineSession* Session);
