Store ID missmatch interface change
Adds an overload for IsAllowedToPurchase taking a reference to a FOnlineError which can provide detail on why a user may be dissallowed.
part 1

diff --git a/Engine/Plugins/Online/OnlineSubsystem/Source/Public/Interfaces/OnlinePurchaseInterface.h b/Engine/Plugins/Online/OnlineSubsystem/Source/Public/Interfaces/OnlinePurchaseInterface.h
--- a/Engine/Plugins/Online/OnlineSubsystem/Source/Public/Interfaces/OnlinePurchaseInterface.h
+++ b/Engine/Plugins/Online/OnlineSubsystem/Source/Public/Interfaces/OnlinePurchaseInterface.h
@@ -263,6 +263,17 @@ public:
 	 */
 	virtual bool IsAllowedToPurchase(const FUniqueNetId& UserId) = 0;
 
+	/**
+	 * Determine if user is allowed to purchase from store
+	 *
+	 * @param UserId user initiating the request
+	 * 
+	 * @param Error an error result indicating why the the user isn't allowed when relevant and avaialable.
+	 *
+	 * @return true if user can make a purchase
+	 */
+	virtual bool IsAllowedToPurchase(const FUniqueNetId& UserId, FOnlineError& Error) = 0;
+
 	/**
 	 * Initiate the checkout process for purchasing offers via payment
 	 *
diff --git a/Engine/Plugins/Online/OnlineSubsystem/Source/Public/OnlineError.h b/Engine/Plugins/Online/OnlineSubsystem/Source/Public/OnlineError.h
--- a/Engine/Plugins/Online/OnlineSubsystem/Source/Public/OnlineError.h
+++ b/Engine/Plugins/Online/OnlineSubsystem/Source/Public/OnlineError.h
@@ -56,7 +56,7 @@ enum class EOnlineErrorResult : uint8
 	FailExtended,
 	/** No game session found */
 	NoGameSession,
-	/** Failed due to mismatched missing user between title and store */
+	/** Failed due to mismatched or missing user between title and store */
 	MismatchedUser,
 
 	/** Default state */
diff --git a/Engine/Plugins/Online/OnlineSubsystemEOS/Source/OnlineSubsystemEOS/Private/OnlineStoreEOS.h b/Engine/Plugins/Online/OnlineSubsystemEOS/Source/OnlineSubsystemEOS/Private/OnlineStoreEOS.h
--- a/Engine/Plugins/Online/OnlineSubsystemEOS/Source/OnlineSubsystemEOS/Private/OnlineStoreEOS.h
+++ b/Engine/Plugins/Online/OnlineSubsystemEOS/Source/OnlineSubsystemEOS/Private/OnlineStoreEOS.h
@@ -35,6 +35,12 @@ public:
 
 // Begin IOnlinePurchase
 	virtual bool IsAllowedToPurchase(const FUniqueNetId& UserId) override { return true; }
+	virtual bool IsAllowedToPurchase(const FUniqueNetId& UserId, FOnlineError& Error) override 
+	{ 
+		Error = FOnlineError(EOnlineErrorResult::Success);
+		return true; 
+	}
+
 	virtual void Checkout(const FUniqueNetId& UserId, const FPurchaseCheckoutRequest& CheckoutRequest, const FOnPurchaseCheckoutComplete& Delegate) override;
 	virtual void Checkout(const FUniqueNetId& UserId, const FPurchaseCheckoutRequest& CheckoutRequest, const FOnPurchaseReceiptlessCheckoutComplete& Delegate) override;
 	virtual void FinalizePurchase(const FUniqueNetId& UserId, const FString& ReceiptId) override;
diff --git a/Engine/Plugins/Online/OnlineSubsystemNull/Source/Private/OnlinePurchaseInterfaceNull.cpp b/Engine/Plugins/Online/OnlineSubsystemNull/Source/Private/OnlinePurchaseInterfaceNull.cpp
--- a/Engine/Plugins/Online/OnlineSubsystemNull/Source/Private/OnlinePurchaseInterfaceNull.cpp
+++ b/Engine/Plugins/Online/OnlineSubsystemNull/Source/Private/OnlinePurchaseInterfaceNull.cpp
@@ -20,6 +20,8 @@ namespace
 	}
 }
 
+#define ONLINE_ERROR_NAMESPACE "errors.com.epicgames.oss.purchase"
+
 FOnlinePurchaseNull::FOnlinePurchaseNull(FOnlineSubsystemNull& InNullSubsystem)
 	: NullSubsystem(InNullSubsystem)
 {
@@ -49,6 +51,12 @@ bool FOnlinePurchaseNull::IsAllowedToPurchase(const FUniqueNetId& UserId)
 	return true;
 }
 
+bool FOnlinePurchaseNull::IsAllowedToPurchase(const FUniqueNetId& UserId, FOnlineError& Error)
+{
+	Error = FOnlineError(EOnlineErrorResult::Success);
+	return true;
+}
+
 void FOnlinePurchaseNull::Checkout(const FUniqueNetId& UserId, const FPurchaseCheckoutRequest& CheckoutRequest, const FOnPurchaseCheckoutComplete& Delegate)
 {
 	// Lambda to wrap calling our delegate with an error and logging the message
diff --git a/Engine/Plugins/Online/OnlineSubsystemNull/Source/Private/OnlinePurchaseInterfaceNull.h b/Engine/Plugins/Online/OnlineSubsystemNull/Source/Private/OnlinePurchaseInterfaceNull.h
--- a/Engine/Plugins/Online/OnlineSubsystemNull/Source/Private/OnlinePurchaseInterfaceNull.h
+++ b/Engine/Plugins/Online/OnlineSubsystemNull/Source/Private/OnlinePurchaseInterfaceNull.h
@@ -22,6 +22,7 @@ public:
 public:
 	//~ Begin IOnlinePurchase Interface
 	virtual bool IsAllowedToPurchase(const FUniqueNetId& UserId) override;
+	virtual bool IsAllowedToPurchase(const FUniqueNetId& UserId, FOnlineError& Error) override;
 	virtual void Checkout(const FUniqueNetId& UserId, const FPurchaseCheckoutRequest& CheckoutRequest, const FOnPurchaseCheckoutComplete& Delegate) override;
 	virtual void Checkout(const FUniqueNetId& UserId, const FPurchaseCheckoutRequest& CheckoutRequest, const FOnPurchaseReceiptlessCheckoutComplete& Delegate) override;
 	virtual void FinalizePurchase(const FUniqueNetId& UserId, const FString& ReceiptId) override;
diff --git a/Engine/Plugins/Online/OnlineSubsystemSteam/Source/Private/OnlinePurchaseInterfaceSteam.cpp b/Engine/Plugins/Online/OnlineSubsystemSteam/Source/Private/OnlinePurchaseInterfaceSteam.cpp
--- a/Engine/Plugins/Online/OnlineSubsystemSteam/Source/Private/OnlinePurchaseInterfaceSteam.cpp
+++ b/Engine/Plugins/Online/OnlineSubsystemSteam/Source/Private/OnlinePurchaseInterfaceSteam.cpp
@@ -30,6 +30,12 @@ bool FOnlinePurchaseSteam::IsAllowedToPurchase(const FUniqueNetId& UserId)
 	return true;
 }
 
+bool FOnlinePurchaseSteam::IsAllowedToPurchase(const FUniqueNetId& UserId, FOnlineError& Error)
+{	
+	Error = FOnlineError(EOnlineErrorResult::Success);
+	return true;
+}
+
 void FOnlinePurchaseSteam::Checkout(const FUniqueNetId& UserId, const FPurchaseCheckoutRequest& CheckoutRequest, const FOnPurchaseCheckoutComplete& Delegate)
 {
 	if(!ServerLink.IsValid())
diff --git a/Engine/Plugins/Online/OnlineSubsystemSteam/Source/Public/OnlinePurchaseInterfaceSteam.h b/Engine/Plugins/Online/OnlineSubsystemSteam/Source/Public/OnlinePurchaseInterfaceSteam.h
--- a/Engine/Plugins/Online/OnlineSubsystemSteam/Source/Public/OnlinePurchaseInterfaceSteam.h
+++ b/Engine/Plugins/Online/OnlineSubsystemSteam/Source/Public/OnlinePurchaseInterfaceSteam.h
@@ -108,6 +108,7 @@ public:
 public:
 	//~ Begin IOnlinePurchase Interface
 	virtual bool IsAllowedToPurchase(const FUniqueNetId& UserId) override;
+	virtual bool IsAllowedToPurchase(const FUniqueNetId& UserId, FOnlineError& Error) override;
 	virtual void Checkout(const FUniqueNetId& UserId, const FPurchaseCheckoutRequest& CheckoutRequest, const FOnPurchaseCheckoutComplete& Delegate) override;
 	virtual void Checkout(const FUniqueNetId& UserId, const FPurchaseCheckoutRequest& CheckoutRequest, const FOnPurchaseReceiptlessCheckoutComplete& Delegate) override;
 	virtual void FinalizePurchase(const FUniqueNetId& UserId, const FString& ReceiptId) override;
diff --git a/Engine/Plugins/Online/OnlineSubsystemTencent/Source/Private/OnlinePurchaseTencent.cpp b/Engine/Plugins/Online/OnlineSubsystemTencent/Source/Private/OnlinePurchaseTencent.cpp
--- a/Engine/Plugins/Online/OnlineSubsystemTencent/Source/Private/OnlinePurchaseTencent.cpp
+++ b/Engine/Plugins/Online/OnlineSubsystemTencent/Source/Private/OnlinePurchaseTencent.cpp
@@ -21,6 +21,12 @@ bool FOnlinePurchaseTencent::IsAllowedToPurchase(const FUniqueNetId& UserId)
 	return bResult;
 }
 
+bool FOnlinePurchaseTencent::IsAllowedToPurchase(const FUniqueNetId& UserId, FOnlineError& Error)
+{
+	Error = FOnlineError(EOnlineErrorResult::Success);
+	return true;
+}
+
 void FOnlinePurchaseTencent::Checkout(const FUniqueNetId& UserId, const FPurchaseCheckoutRequest& CheckoutRequest, const FOnPurchaseCheckoutComplete& Delegate)
 {
 	FOnlineError Result(true);
diff --git a/Engine/Plugins/Online/OnlineSubsystemTencent/Source/Private/OnlinePurchaseTencent.h b/Engine/Plugins/Online/OnlineSubsystemTencent/Source/Private/OnlinePurchaseTencent.h
--- a/Engine/Plugins/Online/OnlineSubsystemTencent/Source/Private/OnlinePurchaseTencent.h
+++ b/Engine/Plugins/Online/OnlineSubsystemTencent/Source/Private/OnlinePurchaseTencent.h
@@ -24,6 +24,7 @@ public:
 
 	// Begin IOnlinePurchase
 	virtual bool IsAllowedToPurchase(const FUniqueNetId& UserId) override;
+	virtual bool IsAllowedToPurchase(const FUniqueNetId& UserId, FOnlineError& Error) override;
 	virtual void Checkout(const FUniqueNetId& UserId, const FPurchaseCheckoutRequest& CheckoutRequest, const FOnPurchaseCheckoutComplete& Delegate) override;
 	virtual void Checkout(const FUniqueNetId& UserId, const FPurchaseCheckoutRequest& CheckoutRequest, const FOnPurchaseReceiptlessCheckoutComplete& Delegate) override;
 	virtual void FinalizePurchase(const FUniqueNetId& UserId, const FString& ReceiptId) override;
