Add BucketID to default Lobby search parameters

[RN] Added the BucketID to the default EOS Lobby search parameters to match EOS Session behaviour.
#Synced-CL 48029169

diff --git a/Engine/Plugins/Online/OnlineSubsystemEOS/Source/OnlineSubsystemEOS/Private/OnlineSessionEOS.cpp b/Engine/Plugins/Online/OnlineSubsystemEOS/Source/OnlineSubsystemEOS/Private/OnlineSessionEOS.cpp
--- a/Engine/Plugins/Online/OnlineSubsystemEOS/Source/OnlineSubsystemEOS/Private/OnlineSessionEOS.cpp
+++ b/Engine/Plugins/Online/OnlineSubsystemEOS/Source/OnlineSubsystemEOS/Private/OnlineSessionEOS.cpp
@@ -4670,13 +4670,17 @@ uint32 FOnlineSessionEOS::FindLobbySession(int32 SearchingPlayerNum, const TShar
 	EOS_Lobby_CreateLobbySearchOptions CreateLobbySearchOptions = { 0 };
 	CreateLobbySearchOptions.ApiVersion = 1;
 	UE_EOS_CHECK_API_MISMATCH(EOS_LOBBY_CREATELOBBYSEARCH_API_LATEST, 1);
-	CreateLobbySearchOptions.MaxResults = FMath::Clamp(SearchSettings->MaxSearchResults, 0, EOS_SESSIONS_MAX_SEARCH_RESULTS);
+	CreateLobbySearchOptions.MaxResults = FMath::Clamp(SearchSettings->MaxSearchResults, 0, EOS_LOBBY_MAX_SEARCH_RESULTS);
 
 	EOS_HLobbySearch LobbySearchHandle;
 
 	EOS_EResult SearchResult = EOS_Lobby_CreateLobbySearch(LobbyHandle, &CreateLobbySearchOptions, &LobbySearchHandle);
 	if (SearchResult == EOS_EResult::EOS_Success)
 	{
+		const auto BucketIdUtf8 = StringCast<UTF8CHAR>(*GetBucketId(*SearchSettings));
+		FLobbyAttributeOptions BucketIdAttribute(EOS_LOBBY_SEARCH_BUCKET_ID, (const char*)BucketIdUtf8.Get());
+		AddLobbySearchAttribute(LobbySearchHandle, &BucketIdAttribute, EOS_EOnlineComparisonOp::EOS_CO_EQUAL);
+
 		// We add the search parameters
 		for (FSearchParams::TConstIterator It(SearchSettings->QuerySettings.SearchParams); It; ++It)
 		{
