Added a new Ecom flow to fix several issues with the existing flow. The new flow can be used by setting bUseNewEcomFlow in config.
#Synced-CL 48736045

diff --git a/Engine/Plugins/Online/OnlineSubsystemEOS/Source/OnlineSubsystemEOS/Private/EOSSettings.cpp b/Engine/Plugins/Online/OnlineSubsystemEOS/Source/OnlineSubsystemEOS/Private/EOSSettings.cpp
--- a/Engine/Plugins/Online/OnlineSubsystemEOS/Source/OnlineSubsystemEOS/Private/EOSSettings.cpp
+++ b/Engine/Plugins/Online/OnlineSubsystemEOS/Source/OnlineSubsystemEOS/Private/EOSSettings.cpp
@@ -148,6 +148,7 @@ FEOSSettings::FEOSSettings()
 	, bUseEAS(false)
 	, bUseEOSConnect(false)
 	, bUseEOSRTC(true)
+	, bUseNewEcomFlow(false)
 PRAGMA_DISABLE_DEPRECATION_WARNINGS
 	, bUseNewLoginFlow(false)
 PRAGMA_ENABLE_DEPRECATION_WARNINGS
@@ -198,6 +199,7 @@ const FEOSSettings& UEOSSettings::ManualGetSettings()
 		GConfig->GetBool(INI_SECTION, TEXT("bUseEOSConnect"), GCachedSettings->bUseEOSConnect, GEngineIni);
 		GConfig->GetBool(INI_SECTION, TEXT("bUseEOSRTC"), GCachedSettings->bUseEOSRTC, GEngineIni);
 		GConfig->GetBool(INI_SECTION, TEXT("bUseNamedPlatformConfig"), GCachedSettings->bUseNamedPlatformConfig, GEngineIni);
+		GConfig->GetBool(INI_SECTION, TEXT("bUseNewEcomFlow"), GCachedSettings->bUseNewEcomFlow, GEngineIni);
 PRAGMA_DISABLE_DEPRECATION_WARNINGS
 		GConfig->GetBool(INI_SECTION, TEXT("bUseNewLoginFlow"), GCachedSettings->bUseNewLoginFlow, GEngineIni);
 PRAGMA_ENABLE_DEPRECATION_WARNINGS
@@ -233,6 +235,7 @@ FEOSSettings UEOSSettings::ToNative() const
 	Native.bUseEOSConnect = bUseEOSConnect;
 	Native.bUseEOSRTC = bUseEOSRTC;
 	Native.bUseNamedPlatformConfig = bUseNamedPlatformConfig;
+	Native.bUseNewEcomFlow = bUseNewEcomFlow;
 PRAGMA_DISABLE_DEPRECATION_WARNINGS
 	Native.bUseNewLoginFlow = bUseNewLoginFlow;
 PRAGMA_ENABLE_DEPRECATION_WARNINGS
diff --git a/Engine/Plugins/Online/OnlineSubsystemEOS/Source/OnlineSubsystemEOS/Private/OnlineEntitlementsEOS.cpp b/Engine/Plugins/Online/OnlineSubsystemEOS/Source/OnlineSubsystemEOS/Private/OnlineEntitlementsEOS.cpp
new file mode 100644
--- /dev/null
+++ b/Engine/Plugins/Online/OnlineSubsystemEOS/Source/OnlineSubsystemEOS/Private/OnlineEntitlementsEOS.cpp
@@ -0,0 +1,271 @@
+// Copyright Epic Games, Inc. All Rights Reserved.
+
+#include "OnlineEntitlementsEOS.h"
+#include "EOSSettings.h"
+#include "IEOSSDKManager.h"
+
+#if WITH_EOS_SDK
+
+#include "OnlineSubsystemEOS.h"
+#include "OnlineSubsystemEOSPrivate.h"
+#include "UserManagerEOS.h"
+#include "eos_ecom.h"
+
+#define ONLINE_ERROR_NAMESPACE "com.epicgames.oss.eos.error"
+
+
+FOnlineEntitlementsEOS::FOnlineEntitlementsEOS(FOnlineSubsystemEOS* InSubsystem)
+	: EOSSubsystem(InSubsystem)
+{
+	check(EOSSubsystem != nullptr);
+}
+
+typedef TEOSCallback<EOS_Ecom_OnQueryEntitlementsCallback, EOS_Ecom_QueryEntitlementsCallbackInfo, FOnlineEntitlementsEOS> FQueryEntitlementsCallback;
+
+bool FOnlineEntitlementsEOS::QueryEntitlements(const FUniqueNetId& UserId, const FString& Namespace, const FPagedQuery& Page)
+{
+	const FUniqueNetIdEOS& UserEOSId = FUniqueNetIdEOS::Cast(UserId);
+	const EOS_EpicAccountId AccountId = UserEOSId.GetEpicAccountId();
+	if (AccountId == nullptr)
+	{
+		UE_LOG_ONLINE(Error, TEXT("[%hs] EpicAccountId is null"), __FUNCTION__);
+		return false;
+	}
+
+	EOS_Ecom_QueryEntitlementsOptions Options = { };
+	Options.ApiVersion = 3;
+	UE_EOS_CHECK_API_MISMATCH(EOS_ECOM_QUERYENTITLEMENTS_API_LATEST, 3);
+	Options.LocalUserId = AccountId;
+	Options.bIncludeRedeemed = EOS_TRUE;
+	const auto NamespaceUtf8 = StringCast<UTF8CHAR>(*Namespace); 
+	Options.OverrideCatalogNamespace = Namespace.IsEmpty() ? nullptr : (const char*) NamespaceUtf8.Get();
+
+	FQueryEntitlementsCallback* CallbackObj = new FQueryEntitlementsCallback(FOnlineEntitlementsEOSWeakPtr(AsShared()));
+	CallbackObj->CallbackLambda = [this, LambdaUserId = UserId.AsShared(), LambdaNamespace = FString(Namespace)](const EOS_Ecom_QueryEntitlementsCallbackInfo* Data)
+	{
+		if (Data->ResultCode == EOS_EResult::EOS_Success)
+		{
+			EOS_Ecom_GetEntitlementsCountOptions CountOptions = { };
+			CountOptions.ApiVersion = 1;
+			UE_EOS_CHECK_API_MISMATCH(EOS_ECOM_GETENTITLEMENTSCOUNT_API_LATEST, 1);
+			CountOptions.LocalUserId = Data->LocalUserId;
+			uint32 Count = EOS_Ecom_GetEntitlementsCount(EOSSubsystem->EcomHandle, &CountOptions);
+
+			TArray<TSharedRef<FOnlineEntitlement>> NewEntitlements;
+			NewEntitlements.Reserve(Count);
+
+			EOS_Ecom_CopyEntitlementByIndexOptions CopyEntitlementOptions = { };
+			CopyEntitlementOptions.ApiVersion = 1;
+			UE_EOS_CHECK_API_MISMATCH(EOS_ECOM_COPYENTITLEMENTBYINDEX_API_LATEST, 1);
+			CopyEntitlementOptions.LocalUserId = Data->LocalUserId;
+
+			for (uint32 Index = 0; Index < Count; Index++)
+			{
+				CopyEntitlementOptions.EntitlementIndex = Index;
+
+				EOS_Ecom_Entitlement* Entitlement = nullptr;
+				EOS_EResult CopyEntitlementResult = EOS_Ecom_CopyEntitlementByIndex(EOSSubsystem->EcomHandle, &CopyEntitlementOptions, &Entitlement);
+				if (CopyEntitlementResult == EOS_EResult::EOS_Success || CopyEntitlementResult == EOS_EResult::EOS_Ecom_EntitlementStale)
+				{
+					TSharedRef<FOnlineEntitlement> NewEntitlement = MakeShared<FOnlineEntitlement>();
+					NewEntitlement->Id = Entitlement->EntitlementId ? StringCast<TCHAR>(Entitlement->EntitlementId).Get() : FString();
+					NewEntitlement->ItemId = Entitlement->CatalogItemId ? StringCast<TCHAR>(Entitlement->CatalogItemId).Get() : FString();
+					NewEntitlement->Name = Entitlement->EntitlementName ? StringCast<TCHAR>(Entitlement->EntitlementName).Get() : FString();
+
+					NewEntitlement->EndDate = Entitlement->EndTimestamp > 0 ? FDateTime::FromUnixTimestamp(Entitlement->EndTimestamp).ToFormattedString(TEXT("%Y-%m-%d_%H-%M-%S")) : TEXT("");
+					NewEntitlement->Namespace = LambdaNamespace;
+					// There is no consumed count in the EOS SDK - This field is used to identify entitlements that have been redeemed
+					NewEntitlement->ConsumedCount = Entitlement->bRedeemed ? 1 : 0;
+
+					EOS_Ecom_CatalogItem* CatalogItem = nullptr;
+					EOS_Ecom_CopyItemByIdOptions CopyCatalogItemOptions = {};
+					CopyCatalogItemOptions.ApiVersion = 1;
+					UE_EOS_CHECK_API_MISMATCH(EOS_ECOM_COPYITEMBYID_API_LATEST, 1);
+					CopyCatalogItemOptions.ItemId = Entitlement->CatalogItemId;
+					CopyCatalogItemOptions.LocalUserId = Data->LocalUserId;
+					EOS_EResult CopyCatalogItemResult = EOS_Ecom_CopyItemById(EOSSubsystem->EcomHandle, &CopyCatalogItemOptions, &CatalogItem);
+
+					if (CopyCatalogItemResult == EOS_EResult::EOS_Success || CopyCatalogItemResult == EOS_EResult::EOS_Ecom_CatalogItemStale)
+					{
+						// EOS_Ecom_CopyItemById will return EOS_NotFound is QueryOffers has NOT been called OR if the catalogItemId doesn't tie back to a specific offer (bundles, season packs, etc...)
+						NewEntitlement->bIsConsumable = CatalogItem->ItemType == EOS_EEcomItemType::EOS_EIT_Consumable;
+						EOS_Ecom_CatalogItem_Release(CatalogItem);
+					}
+					else
+					{
+						UE_LOG_ONLINE(Verbose, TEXT("[FOnlineEntitlementsEOS::QueryEntitlements] EOS_Ecom_CopyItemById: failed with error (%s)"), *LexToString(CopyCatalogItemResult));
+						// Default bIsConsumbale to true
+						NewEntitlement->bIsConsumable = true;
+					}
+					NewEntitlements.Add(NewEntitlement);
+				}
+				else
+				{
+					UE_LOG_ONLINE(Error, TEXT("[FOnlineEntitlementsEOS::QueryEntitlements] EOS_Ecom_CopyEntitlementByIndex: failed with error (%s)"), *LexToString(CopyEntitlementResult));
+					continue;
+				}
+
+				EOS_Ecom_Entitlement_Release(Entitlement);
+			}
+
+			CachedEntitlementsMap.FindOrAdd(LambdaUserId) = MoveTemp(NewEntitlements);
+			TriggerOnQueryEntitlementsCompleteDelegates(true, *LambdaUserId, LambdaNamespace, "");
+		}
+		else
+		{
+			FString ErrorString = FString::Printf(TEXT("result code (%s)"), *LexToString(Data->ResultCode));
+			UE_LOG_ONLINE(Error, TEXT("[FOnlineEntitlementsEOS::QueryEntitlements] %s"), *ErrorString);
+			TriggerOnQueryEntitlementsCompleteDelegates(false,*LambdaUserId, LambdaNamespace, ErrorString);
+		}
+	};
+	
+	EOS_Ecom_QueryEntitlements(EOSSubsystem->EcomHandle, &Options, CallbackObj, CallbackObj->GetCallbackPtr());
+	return true;
+}
+
+TSharedPtr<FOnlineEntitlement> FOnlineEntitlementsEOS::GetEntitlement(const FUniqueNetId& UserId, const FUniqueEntitlementId& EntitlementId)
+{
+	const FUniqueNetIdEOS& UserEOSId = FUniqueNetIdEOS::Cast(UserId);
+	const EOS_EpicAccountId AccountId = UserEOSId.GetEpicAccountId();
+	if (AccountId == nullptr)
+	{
+		UE_LOG_ONLINE(Error, TEXT("[%hs] EpicAccountId is null"), __FUNCTION__);
+		return nullptr;
+	}
+
+	const TArray<TSharedRef<FOnlineEntitlement>>* CachedUserEntitlements = CachedEntitlementsMap.Find(UserId.AsShared());
+	if (!CachedUserEntitlements || CachedUserEntitlements->IsEmpty())
+	{
+		UE_LOG_ONLINE(Warning, TEXT("[%hs] CachedEntitlements empty for user %s, call QueryEntitlements"), __FUNCTION__, *UserId.ToString());
+		return nullptr;
+	}
+
+	TSharedPtr<FOnlineEntitlement> Entitlement = nullptr;
+	for (const TSharedRef<FOnlineEntitlement>& CachedUserEntitlement : *CachedUserEntitlements)
+	{
+		if (CachedUserEntitlement->Id == EntitlementId)
+		{
+			Entitlement = CachedUserEntitlement;
+			break;
+		}
+	}
+
+	return Entitlement;
+}
+
+TSharedPtr<FOnlineEntitlement> FOnlineEntitlementsEOS::GetItemEntitlement(const FUniqueNetId& UserId, const FString& ItemId)
+{
+	const FUniqueNetIdEOS& UserEOSId = FUniqueNetIdEOS::Cast(UserId);
+	const EOS_EpicAccountId AccountId = UserEOSId.GetEpicAccountId();
+	if (AccountId == nullptr)
+	{
+		UE_LOG_ONLINE(Error, TEXT("[%hs] EpicAccountId is null"), __FUNCTION__);
+		return nullptr;
+	}
+
+	const TArray<TSharedRef<FOnlineEntitlement>>* CachedUserEntitlements = CachedEntitlementsMap.Find(UserId.AsShared());
+	if (!CachedUserEntitlements || CachedUserEntitlements->IsEmpty())
+	{
+		UE_LOG_ONLINE(Warning, TEXT("[%hs] CachedEntitlements empty for user %s, call QueryEntitlements"), __FUNCTION__, *UserId.ToString());
+		return nullptr;
+	}
+
+	TSharedPtr<FOnlineEntitlement> Entitlement = nullptr;
+	for (const TSharedRef<FOnlineEntitlement>& CachedUserEntitlement : *CachedUserEntitlements)
+	{
+		if (CachedUserEntitlement->ItemId == ItemId)
+		{
+			Entitlement = CachedUserEntitlement;
+			break;
+		}
+	}
+
+	return Entitlement;
+}
+
+void FOnlineEntitlementsEOS::GetAllEntitlements(const FUniqueNetId& UserId, const FString& Namespace, TArray<TSharedRef<FOnlineEntitlement>>& OutUserEntitlements)
+{
+	OutUserEntitlements.Empty();
+
+	const FUniqueNetIdEOS& UserEOSId = FUniqueNetIdEOS::Cast(UserId);
+	const EOS_EpicAccountId AccountId = UserEOSId.GetEpicAccountId();
+	if (AccountId == nullptr)
+	{
+		UE_LOG_ONLINE(Error, TEXT("[%hs] EpicAccountId is null"), __FUNCTION__);
+		return;
+	}
+
+	const TArray<TSharedRef<FOnlineEntitlement>>* CachedUserEntitlements = CachedEntitlementsMap.Find(UserId.AsShared());
+	if (!CachedUserEntitlements || CachedUserEntitlements->IsEmpty())
+	{
+		UE_LOG_ONLINE(Warning, TEXT("[%hs] CachedEntitlements empty for user %s, call QueryEntitlements"), __FUNCTION__, *UserId.ToString());
+		return;
+	}
+
+	if (Namespace.IsEmpty())
+	{
+		OutUserEntitlements = *CachedUserEntitlements;
+	}
+	else
+	{
+		for (const TSharedRef<FOnlineEntitlement>& CachedUserEntitlement: *CachedUserEntitlements)
+		{
+			if (CachedUserEntitlement->Namespace == Namespace)
+			{
+				OutUserEntitlements.Add(CachedUserEntitlement);
+			}
+		}
+	}
+}
+
+bool FOnlineEntitlementsEOS::HandleEntitlementsExec(UWorld* InWorld, const TCHAR* Cmd, FOutputDevice& Ar)
+{
+	if (!EOSSubsystem || !EOSSubsystem->UserManager.IsValid() || !EOSSubsystem->EntitlementsInterfacePtr.IsValid())
+	{
+		UE_LOG_ONLINE(Error, TEXT("[%hs] EOSSubsystem or UserManager or EntitlementsInterface invalid"), __FUNCTION__);
+		return false;
+	}
+
+	if (FParse::Command(&Cmd, TEXT("ENTITLEMENTS")))
+	{
+		TSharedPtr<FDelegateHandle> DelegateHandle = MakeShared<FDelegateHandle>();
+			
+		*DelegateHandle = AddOnQueryEntitlementsCompleteDelegate_Handle(FOnQueryEntitlementsComplete::FDelegate::CreateLambda([this, DelegateHandle](bool bWasSuccessful, const FUniqueNetId& UserId, const FString& Namespace, const FString& Error)
+		{
+			if (DelegateHandle.IsValid())
+			{
+				ClearOnQueryEntitlementsCompleteDelegate_Handle(*DelegateHandle);
+				DelegateHandle->Reset();
+			}
+
+			if (bWasSuccessful)
+			{
+				const TArray<TSharedRef<FOnlineEntitlement>>* CachedUserEntitlements = CachedEntitlementsMap.Find(UserId.AsShared());
+				if (!CachedUserEntitlements || CachedUserEntitlements->IsEmpty())
+				{
+					UE_LOG_ONLINE(Log, TEXT("[FOnlineEntitlementsEOS::HandleEntitlementsExec] CachedEntitlements empty for user %s, call QueryEntitlements"), *UserId.ToString());
+					return;
+				}
+
+				for (const TSharedRef<FOnlineEntitlement>& CachedUserEntitlement : *CachedUserEntitlements)
+				{
+					UE_LOG_ONLINE(Log, TEXT("[FOnlineEntitlementsEOS::HandleEntitlementsExec] EntitlementId: %s"), *CachedUserEntitlement->Id);
+					UE_LOG_ONLINE(Log, TEXT("[FOnlineEntitlementsEOS::HandleEntitlementsExec] ItemId: %s"), *CachedUserEntitlement->ItemId);
+					UE_LOG_ONLINE(Log, TEXT("[FOnlineEntitlementsEOS::HandleEntitlementsExec] Name: %s"), *CachedUserEntitlement->Name);
+					UE_LOG_ONLINE(Log, TEXT("[FOnlineEntitlementsEOS::HandleEntitlementsExec] bIsConsumable: %s"), *LexToString(CachedUserEntitlement->bIsConsumable));
+				}
+			}
+			else
+			{
+				UE_LOG_ONLINE(Error, TEXT("[FOnlineEntitlementsEOS::HandleEntitlementsExec] error (%s)"), *Error);
+			}
+		}));
+
+		QueryEntitlements(*EOSSubsystem->UserManager->GetLocalUniqueNetIdEOS(EOSSubsystem->UserManager->GetDefaultLocalUser()), "");
+		return true;
+	}
+	return false;
+}
+
+#endif
+
diff --git a/Engine/Plugins/Online/OnlineSubsystemEOS/Source/OnlineSubsystemEOS/Private/OnlineEntitlementsEOS.h b/Engine/Plugins/Online/OnlineSubsystemEOS/Source/OnlineSubsystemEOS/Private/OnlineEntitlementsEOS.h
new file mode 100644
--- /dev/null
+++ b/Engine/Plugins/Online/OnlineSubsystemEOS/Source/OnlineSubsystemEOS/Private/OnlineEntitlementsEOS.h
@@ -0,0 +1,48 @@
+// Copyright Epic Games, Inc. All Rights Reserved.
+
+#pragma once
+
+#include "CoreMinimal.h"
+#include "Online/CoreOnline.h"
+#include "Interfaces/OnlineEntitlementsInterface.h"
+#include "OnlineSubsystemEOSTypes.h"
+
+#if WITH_EOS_SDK
+#include "eos_ecom_types.h"
+
+class UWorld;
+
+/**
+ * Query entitlements
+ */
+class FOnlineEntitlementsEOS :
+	public IOnlineEntitlements,
+	public TSharedFromThis<FOnlineEntitlementsEOS, ESPMode::ThreadSafe>
+{
+public:
+	virtual ~FOnlineEntitlementsEOS() = default;
+	
+	virtual bool QueryEntitlements(const FUniqueNetId& UserId, const FString& Namespace, const FPagedQuery& Page = FPagedQuery()) override;
+	virtual TSharedPtr<FOnlineEntitlement> GetEntitlement(const FUniqueNetId& UserId, const FUniqueEntitlementId& EntitlementId) override;
+	virtual TSharedPtr<FOnlineEntitlement> GetItemEntitlement(const FUniqueNetId& UserId, const FString& ItemId) override;
+	virtual void GetAllEntitlements(const FUniqueNetId& UserId, const FString& Namespace, TArray<TSharedRef<FOnlineEntitlement>>& OutUserEntitlements) override;
+	
+	FOnlineEntitlementsEOS(FOnlineSubsystemEOS* InSubsystem);
+	
+	bool HandleEntitlementsExec(UWorld* InWorld, const TCHAR* Cmd, FOutputDevice& Ar);
+
+private:
+	/** Default constructor disabled */
+	FOnlineEntitlementsEOS() = delete;
+
+	/** List of entitlements for the player */
+	TMap<FUniqueNetIdRef, TArray<TSharedRef<FOnlineEntitlement>>> CachedEntitlementsMap;
+
+	/** Reference to the main EOS subsystem */
+	FOnlineSubsystemEOS* EOSSubsystem;
+};
+
+typedef TSharedPtr<FOnlineEntitlementsEOS, ESPMode::ThreadSafe> FOnlineEntitlementsEOSPtr;
+typedef TWeakPtr<FOnlineEntitlementsEOS, ESPMode::ThreadSafe> FOnlineEntitlementsEOSWeakPtr;
+
+#endif
diff --git a/Engine/Plugins/Online/OnlineSubsystemEOS/Source/OnlineSubsystemEOS/Private/OnlinePurchaseEOS.cpp b/Engine/Plugins/Online/OnlineSubsystemEOS/Source/OnlineSubsystemEOS/Private/OnlinePurchaseEOS.cpp
new file mode 100644
--- /dev/null
+++ b/Engine/Plugins/Online/OnlineSubsystemEOS/Source/OnlineSubsystemEOS/Private/OnlinePurchaseEOS.cpp
@@ -0,0 +1,534 @@
+// Copyright Epic Games, Inc. All Rights Reserved.
+
+#include "OnlinePurchaseEOS.h"
+#include "OnlineEntitlementsEOS.h"
+#include "EOSSettings.h"
+#include "IEOSSDKManager.h"
+
+#if WITH_EOS_SDK
+#include "OnlineSubsystemEOS.h"
+#include "OnlineSubsystemEOSPrivate.h"
+#include "UserManagerEOS.h"
+#include "eos_ecom.h"
+
+#define ONLINE_ERROR_NAMESPACE "com.epicgames.oss.eos.error"
+
+
+FOnlinePurchaseEOS::FOnlinePurchaseEOS(FOnlineSubsystemEOS* InSubsystem)
+	: EOSSubsystem(InSubsystem)
+{
+	check(EOSSubsystem != nullptr);
+}
+
+typedef TEOSCallback<EOS_Ecom_OnCheckoutCallback, EOS_Ecom_CheckoutCallbackInfo, FOnlinePurchaseEOS> FPurchaseCheckoutCallback;
+
+void FOnlinePurchaseEOS::Checkout(const FUniqueNetId& UserId, const FPurchaseCheckoutRequest& CheckoutRequest, const FOnPurchaseCheckoutComplete& Delegate)
+{
+	const FUniqueNetIdEOS& UserEOSId = FUniqueNetIdEOS::Cast(UserId);
+	const EOS_EpicAccountId AccountId = UserEOSId.GetEpicAccountId();
+	if (AccountId == nullptr)
+	{
+		UE_LOG_ONLINE(Error, TEXT("[%hs] EpicAccountId is null"), __FUNCTION__);
+		Delegate.ExecuteIfBound(ONLINE_ERROR(EOnlineErrorResult::InvalidUser), MakeShared<FPurchaseReceipt>());
+		return;
+	}
+	if (CheckoutRequest.PurchaseOffers.Num() == 0)
+	{
+		UE_LOG_ONLINE(Error, TEXT("[%hs] PurchaseOffers.Num() is 0"), __FUNCTION__);
+		Delegate.ExecuteIfBound(ONLINE_ERROR(EOnlineErrorResult::InvalidParams), MakeShared<FPurchaseReceipt>());
+		return;
+	}
+	if (CheckoutRequest.PurchaseOffers.Num() > EOS_ECOM_CHECKOUT_MAX_ENTRIES)
+	{
+		UE_LOG_ONLINE(Error, TEXT("[%hs] PurcahseOffer.Num() over maximum of %d "), __FUNCTION__, EOS_ECOM_CHECKOUT_MAX_ENTRIES);
+		Delegate.ExecuteIfBound(ONLINE_ERROR(EOnlineErrorResult::InvalidParams), MakeShared<FPurchaseReceipt>());
+		return;
+	}
+
+	const int32 NumItems = CheckoutRequest.PurchaseOffers.Num();
+	TArray<EOS_Ecom_CheckoutEntry> Entries;
+	TArray<FTCHARToUTF8> OfferIdsUtf8;
+	Entries.AddZeroed(NumItems);
+	OfferIdsUtf8.Reserve(NumItems);
+
+	for (int32 Index = 0; Index < NumItems; Index++)
+	{
+		OfferIdsUtf8.Emplace(*CheckoutRequest.PurchaseOffers[Index].OfferId);
+		Entries[Index].ApiVersion = 1;
+		UE_EOS_CHECK_API_MISMATCH(EOS_ECOM_CHECKOUTENTRY_API_LATEST, 1);
+		Entries[Index].OfferId = OfferIdsUtf8.Last().Get();
+	}
+
+	EOS_Ecom_CheckoutOptions Options = { };
+	Options.ApiVersion = 2;
+	UE_EOS_CHECK_API_MISMATCH(EOS_ECOM_CHECKOUT_API_LATEST, 2);
+	Options.LocalUserId = AccountId;
+	Options.EntryCount = NumItems;
+	Options.Entries = (const EOS_Ecom_CheckoutEntry*)Entries.GetData();
+	Options.OverrideCatalogNamespace = nullptr;
+	Options.PreferredOrientation = EOS_ECheckoutOrientation::EOS_ECO_Default;
+
+	FPurchaseCheckoutCallback* CallbackObj = new FPurchaseCheckoutCallback(FOnlinePurchaseEOSWeakPtr(AsShared()));
+	CallbackObj->CallbackLambda = [this, OnComplete = FOnPurchaseCheckoutComplete(Delegate)](const EOS_Ecom_CheckoutCallbackInfo* Data)
+	{
+		const EOS_EResult CheckoutResult = Data->ResultCode;
+		if (CheckoutResult != EOS_EResult::EOS_Success)
+		{
+			if (CheckoutResult == EOS_EResult::EOS_Canceled)
+			{
+				UE_LOG_ONLINE(Log, TEXT("[FOnlinePurchaseEOS::Checkout] EOS_Ecom_Checkout was canceled by user"));
+				OnComplete.ExecuteIfBound(ONLINE_ERROR(EOnlineErrorResult::Canceled), MakeShared<FPurchaseReceipt>());
+			}
+			else if (CheckoutResult == EOS_EResult::EOS_Ecom_PurchaseProcessing)
+			{
+				// See header documentation for EOS_Ecom_PurchaseProcessing and online EOS documentation for more information
+				UE_LOG_ONLINE(Log, TEXT("[FOnlinePurchaseEOS::Checkout] QueryEntitlements will need to be called to confirm the purchase. This can be done via UI button"));
+				OnComplete.ExecuteIfBound(ONLINE_ERROR(EOnlineErrorResult::FailExtended, LexToString(CheckoutResult)), MakeShared<FPurchaseReceipt>());
+			}
+			else
+			{
+				UE_LOG_ONLINE(Error, TEXT("[FOnlinePurchaseEOS::Checkout] EOS_Ecom_Checkout: failed with error (%s)"), *LexToString(CheckoutResult));
+				OnComplete.ExecuteIfBound(ONLINE_ERROR(EOnlineErrorResult::Unknown), MakeShared<FPurchaseReceipt>());
+			}
+			return;
+		}
+
+		TSharedRef<FPurchaseReceipt> PurchaseReceipt = MakeShared<FPurchaseReceipt>();
+		PurchaseReceipt->TransactionId = Data->TransactionId ? StringCast<TCHAR>(Data->TransactionId).Get() : FString();
+		PurchaseReceipt->TransactionState = EPurchaseTransactionState::Purchased;
+
+		EOS_Ecom_CopyTransactionByIdOptions CopyTransactionOptions = { };
+		CopyTransactionOptions.ApiVersion = 1;
+		UE_EOS_CHECK_API_MISMATCH(EOS_ECOM_COPYTRANSACTIONBYID_API_LATEST, 1);
+		CopyTransactionOptions.LocalUserId = Data->LocalUserId;
+		CopyTransactionOptions.TransactionId = Data->TransactionId;
+		EOS_Ecom_HTransaction TransactionHandle;
+		const EOS_EResult CopyTransactionResult = EOS_Ecom_CopyTransactionById(EOSSubsystem->EcomHandle, &CopyTransactionOptions, &TransactionHandle);
+		if (CopyTransactionResult == EOS_EResult::EOS_Success)
+		{
+			EOS_Ecom_Transaction_GetEntitlementsCountOptions EntitlementCountOptions = { };
+			EntitlementCountOptions.ApiVersion = 1;
+			UE_EOS_CHECK_API_MISMATCH(EOS_ECOM_TRANSACTION_GETENTITLEMENTSCOUNT_API_LATEST, 1);
+			const uint32_t EntitlementsCount = EOS_Ecom_Transaction_GetEntitlementsCount(TransactionHandle, &EntitlementCountOptions);
+
+			for (uint32_t i = 0; i < EntitlementsCount; i++)
+			{
+				EOS_Ecom_Transaction_CopyEntitlementByIndexOptions Options = { };
+				Options.ApiVersion = 1;
+				UE_EOS_CHECK_API_MISMATCH(EOS_ECOM_TRANSACTION_COPYENTITLEMENTBYINDEX_API_LATEST, 1);
+				Options.EntitlementIndex = i;
+				EOS_Ecom_Entitlement* Entitlement;
+				const EOS_EResult CopyEntitlementResult = EOS_Ecom_Transaction_CopyEntitlementByIndex(TransactionHandle, &Options, &Entitlement);
+
+				// The Receipt structure is 1 EOS_Ecom_Entitlement (mapped to FReceiptOfferEntry) has 1 line item
+				if (CopyEntitlementResult == EOS_EResult::EOS_Success)
+				{
+					// OfferId is set to an empty string as the OfferId is not retrievable in the checkout callback
+					PurchaseReceipt->AddReceiptOffer(FOfferNamespace(), "", 1);
+					FPurchaseReceipt::FLineItemInfo& LineItem = PurchaseReceipt->ReceiptOffers[i].LineItems.Emplace_GetRef();
+					// The item that is owned is the EOS_Ecom_CatalogItemId
+					LineItem.UniqueId = Entitlement->CatalogItemId ? StringCast<TCHAR>(Entitlement->CatalogItemId).Get() : FUniqueEntitlementId();
+					// If the item can be redeemed, store the EntitlementId so it can be passed to FinalizeReceiptValidationInfo()
+					FString EntitlementId = Entitlement->EntitlementId ? StringCast<TCHAR>(Entitlement->EntitlementId).Get() : FString();
+					LineItem.ValidationInfo = Entitlement->bRedeemed ? FString() : EntitlementId;
+
+					EOS_Ecom_CatalogItem* CatalogItem = nullptr;
+					EOS_Ecom_CopyItemByIdOptions CopyCatalogItemOptions = {};
+					CopyCatalogItemOptions.ApiVersion = 1;
+					UE_EOS_CHECK_API_MISMATCH(EOS_ECOM_COPYITEMBYID_API_LATEST, 1);
+					CopyCatalogItemOptions.ItemId = Entitlement->CatalogItemId;
+					CopyCatalogItemOptions.LocalUserId = Data->LocalUserId;
+					EOS_EResult CopyCatalogItemResult = EOS_Ecom_CopyItemById(EOSSubsystem->EcomHandle, &CopyCatalogItemOptions, &CatalogItem);
+
+					if (CopyCatalogItemResult == EOS_EResult::EOS_Success || CopyCatalogItemResult == EOS_EResult::EOS_Ecom_CatalogItemStale)
+					{
+						LineItem.ItemName = CatalogItem->TitleText ? StringCast<TCHAR>(CatalogItem->TitleText).Get() : FString();
+						EOS_Ecom_CatalogItem_Release(CatalogItem);
+					}
+					else
+					{
+						// EOS_Ecom_CopyItemById will return EOS_NotFound is QueryOffers has NOT been called OR if the EOS_Ecom_CatalogItemId doesn't tie back to a specific offer (bundles, season packs, etc...)
+						UE_LOG_ONLINE(Verbose, TEXT("[FOnlinePurchaseEOS::Checkout] EOS_Ecom_CopyItemById: failed with error (%s)"), *LexToString(CopyCatalogItemResult));
+					}
+					EOS_Ecom_Entitlement_Release(Entitlement);
+				}
+				else
+				{
+					UE_LOG_ONLINE(Error, TEXT("[FOnlinePurchaseEOS::Checkout] EOS_Ecom_Transaction_CopyEntitlementByIndex: failed with error (%s)"), *LexToString(CopyEntitlementResult));
+					OnComplete.ExecuteIfBound(ONLINE_ERROR(EOnlineErrorResult::Unknown), MakeShared<FPurchaseReceipt>());
+					EOS_Ecom_Transaction_Release(TransactionHandle);
+					return;
+				}
+			}
+			EOS_Ecom_Transaction_Release(TransactionHandle);
+			OnComplete.ExecuteIfBound(ONLINE_ERROR(EOnlineErrorResult::Success), PurchaseReceipt);
+		}
+		else
+		{
+			OnComplete.ExecuteIfBound(ONLINE_ERROR(EOnlineErrorResult::Unknown), MakeShared<FPurchaseReceipt>());
+		}
+	};
+
+	EOS_Ecom_Checkout(EOSSubsystem->EcomHandle, &Options, CallbackObj, CallbackObj->GetCallbackPtr());
+}
+
+void FOnlinePurchaseEOS::Checkout(const FUniqueNetId& UserId, const FPurchaseCheckoutRequest& CheckoutRequest, const FOnPurchaseReceiptlessCheckoutComplete& Delegate)
+{
+	const FUniqueNetIdEOS& UserEOSId = FUniqueNetIdEOS::Cast(UserId);
+	const EOS_EpicAccountId AccountId = UserEOSId.GetEpicAccountId();
+	if (AccountId == nullptr)
+	{
+		UE_LOG_ONLINE(Error, TEXT("[%hs] EpicAccountId is null"), __FUNCTION__);
+		Delegate.ExecuteIfBound(ONLINE_ERROR(EOnlineErrorResult::InvalidUser));
+		return;
+	}
+	if (CheckoutRequest.PurchaseOffers.Num() == 0)
+	{
+		UE_LOG_ONLINE(Error, TEXT("[%hs] PurchaseOffers.Num() is 0"), __FUNCTION__);
+		Delegate.ExecuteIfBound(ONLINE_ERROR(EOnlineErrorResult::InvalidParams));
+		return;
+	}
+	if (CheckoutRequest.PurchaseOffers.Num() > EOS_ECOM_CHECKOUT_MAX_ENTRIES)
+	{
+		UE_LOG_ONLINE(Error, TEXT("[%hs] PurcahseOffer.Num() over maximum of %d "), __FUNCTION__, EOS_ECOM_CHECKOUT_MAX_ENTRIES);
+		Delegate.ExecuteIfBound(ONLINE_ERROR(EOnlineErrorResult::InvalidParams));
+		return;
+	}
+
+	const int32 NumItems = CheckoutRequest.PurchaseOffers.Num();
+	TArray<EOS_Ecom_CheckoutEntry> Entries;
+	TArray<FTCHARToUTF8> OfferIdsUtf8;
+	Entries.AddZeroed(NumItems);
+	OfferIdsUtf8.Reserve(NumItems);
+
+	for (int32 Index = 0; Index < NumItems; Index++)
+	{
+		OfferIdsUtf8.Emplace(*CheckoutRequest.PurchaseOffers[Index].OfferId);
+		Entries[Index].ApiVersion = 1;
+		UE_EOS_CHECK_API_MISMATCH(EOS_ECOM_CHECKOUTENTRY_API_LATEST, 1);
+		Entries[Index].OfferId = OfferIdsUtf8.Last().Get();
+	}
+
+	EOS_Ecom_CheckoutOptions Options = { };
+	Options.ApiVersion = 2;
+	UE_EOS_CHECK_API_MISMATCH(EOS_ECOM_CHECKOUT_API_LATEST, 2);
+	Options.LocalUserId = AccountId;
+	Options.EntryCount = NumItems;
+	Options.Entries = (const EOS_Ecom_CheckoutEntry*)Entries.GetData();
+	Options.OverrideCatalogNamespace = nullptr;
+	Options.PreferredOrientation = EOS_ECheckoutOrientation::EOS_ECO_Default;
+
+	FPurchaseCheckoutCallback* CallbackObj = new FPurchaseCheckoutCallback(FOnlinePurchaseEOSWeakPtr(AsShared()));
+	CallbackObj->CallbackLambda = [this, OnComplete = FOnPurchaseReceiptlessCheckoutComplete(Delegate)](const EOS_Ecom_CheckoutCallbackInfo* Data)
+	{
+		const EOS_EResult CheckoutResult = Data->ResultCode;
+		if (CheckoutResult != EOS_EResult::EOS_Success)
+		{
+			if (CheckoutResult == EOS_EResult::EOS_Canceled)
+			{
+				UE_LOG_ONLINE(Log, TEXT("[FOnlinePurchaseEOS::Checkout] EOS_Ecom_Checkout was canceled by user"));
+				OnComplete.ExecuteIfBound(ONLINE_ERROR(EOnlineErrorResult::Canceled));
+			}
+			else if (CheckoutResult == EOS_EResult::EOS_Ecom_PurchaseProcessing)
+			{
+				// See header documentation for EOS_Ecom_PurchaseProcessing and online EOS documentation for more information
+				UE_LOG_ONLINE(Log, TEXT("[FOnlinePurchaseEOS::Checkout] QueryEntitlements will need to be called to confirm the purchase. This can be done via a UI button"));
+				OnComplete.ExecuteIfBound(ONLINE_ERROR(EOnlineErrorResult::FailExtended, LexToString(CheckoutResult)));
+			}
+			else
+			{
+				UE_LOG_ONLINE(Error, TEXT("[FOnlinePurchaseEOS::Checkout] EOS_Ecom_Checkout: failed with error (%s)"), *LexToString(CheckoutResult));
+				OnComplete.ExecuteIfBound(ONLINE_ERROR(EOnlineErrorResult::Unknown));
+			}
+			return;
+		}
+		else
+		{
+			OnComplete.ExecuteIfBound(ONLINE_ERROR(EOnlineErrorResult::Success));
+		}
+	};
+
+	EOS_Ecom_Checkout(EOSSubsystem->EcomHandle, &Options, CallbackObj, CallbackObj->GetCallbackPtr());
+}
+
+void FOnlinePurchaseEOS::FinalizePurchase(const FUniqueNetId& UserId, const FString& ReceiptId)
+{
+	UE_LOG_ONLINE(Verbose, TEXT("[%hs] Not supported. Did you mean FinalizeReceiptValidationInfo?"), __FUNCTION__);
+}
+
+void FOnlinePurchaseEOS::RedeemCode(const FUniqueNetId& UserId, const FRedeemCodeRequest& RedeemCodeRequest, const FOnPurchaseRedeemCodeComplete& Delegate)
+{
+	UE_LOG_ONLINE(Verbose, TEXT("[%hs] Not supported"), __FUNCTION__);
+	static const TSharedRef<FPurchaseReceipt> BlankReceipt(MakeShared<FPurchaseReceipt>());
+	Delegate.ExecuteIfBound(ONLINE_ERROR(EOnlineErrorResult::NotImplemented), BlankReceipt);
+}
+
+typedef TEOSCallback<EOS_Ecom_OnQueryOwnershipBySandboxIdsCallback, EOS_Ecom_QueryOwnershipBySandboxIdsCallbackInfo, FOnlinePurchaseEOS> FPurchaseQueryOwnershipBySandboxIdsCallback;
+
+void FOnlinePurchaseEOS::QueryReceipts(const FUniqueNetId& UserId, bool /* bRestoreReceipts */, const FOnQueryReceiptsComplete& Delegate)
+{
+	const FUniqueNetIdEOS& UserEOSId = FUniqueNetIdEOS::Cast(UserId);
+	const EOS_EpicAccountId AccountId = UserEOSId.GetEpicAccountId();
+	if (AccountId == nullptr)
+	{
+		UE_LOG_ONLINE(Error, TEXT("[%hs] EpicAccountId is null"), __FUNCTION__);
+		Delegate.ExecuteIfBound(ONLINE_ERROR(EOnlineErrorResult::InvalidUser));
+		return;
+	}
+
+	CachedReceipts.Reset();
+
+	// Retrieve the SandboxId from config
+	FString SandboxIdString;
+	const FEOSSettings& EOSSettings = UEOSSettings::GetSettings();
+	if (EOSSettings.bUseNamedPlatformConfig)
+	{
+		const FEOSSDKPlatformConfig* const PlatformConfig = EOSSubsystem->EOSSDKManager->GetPlatformConfig(EOSSettings.PlatformConfigName);
+		if (!PlatformConfig)
+		{
+			UE_LOG_ONLINE(Warning, TEXT("[%hs] GetPlatformConfig failed"), __FUNCTION__);
+			Delegate.ExecuteIfBound(ONLINE_ERROR(EOnlineErrorResult::Unknown));
+			return;
+		}
+		SandboxIdString = PlatformConfig->SandboxId;
+	}
+	else
+	{
+		FEOSArtifactSettings ArtifactSettings;
+		if (!UEOSSettings::GetSelectedArtifactSettings(ArtifactSettings))
+		{
+			UE_LOG_ONLINE(Warning, TEXT("[%hs] GetSelectedArtifactSettings failed"), __FUNCTION__);
+			Delegate.ExecuteIfBound(ONLINE_ERROR(EOnlineErrorResult::Unknown));
+			return;
+		}
+		SandboxIdString = ArtifactSettings.SandboxId;
+	}
+
+	if (SandboxIdString.IsEmpty())
+	{
+		UE_LOG_ONLINE(Warning, TEXT("[%hs] Empty SandboxId"), __FUNCTION__);
+		Delegate.ExecuteIfBound(ONLINE_ERROR(EOnlineErrorResult::Unknown));
+		return;
+	}
+
+	EOS_Ecom_QueryOwnershipBySandboxIdsOptions Options = {};
+	Options.ApiVersion = 1;
+	UE_EOS_CHECK_API_MISMATCH(EOS_ECOM_QUERYOWNERSHIPBYSANDBOXIDSOPTIONS_API_LATEST, 1);
+	Options.LocalUserId = AccountId;
+	Options.SandboxIdsCount = 1;
+	const auto SandboxIdUtf8 = StringCast<UTF8CHAR>(*SandboxIdString);
+	const char* SandboxId = (const char*)SandboxIdUtf8.Get();
+	Options.SandboxIds = &SandboxId;
+
+	FPurchaseQueryOwnershipBySandboxIdsCallback* CallbackObj = new FPurchaseQueryOwnershipBySandboxIdsCallback(FOnlinePurchaseEOSWeakPtr(AsShared()));
+	CallbackObj->CallbackLambda = [this, OnComplete = FOnQueryReceiptsComplete(Delegate)](const EOS_Ecom_QueryOwnershipBySandboxIdsCallbackInfo* Data)
+	{
+		const EOS_EResult Result = Data->ResultCode;
+		if (Result == EOS_EResult::EOS_Success)
+		{
+			UE_LOG_ONLINE(VeryVerbose, TEXT("[FOnlinePurchaseEOS::QueryReceipts] Sandbox Id Ownership Count: %d"), Data->SandboxIdItemOwnershipsCount);
+
+			for (const EOS_Ecom_SandboxIdItemOwnership& SandboxOwnershipEntry : TArrayView<const EOS_Ecom_SandboxIdItemOwnership>(Data->SandboxIdItemOwnerships, Data->SandboxIdItemOwnershipsCount))
+			{
+				UE_LOG_ONLINE(VeryVerbose, TEXT("[FOnlinePurchaseEOS::QueryReceipts] Sandbox Id: %hs"), SandboxOwnershipEntry.SandboxId);
+				UE_LOG_ONLINE(VeryVerbose, TEXT("[FOnlinePurchaseEOS::QueryReceipts] Ownership Count: %d"), SandboxOwnershipEntry.OwnedCatalogItemIdsCount);
+
+				for (uint32_t CatalogItemIndex = 0; CatalogItemIndex < SandboxOwnershipEntry.OwnedCatalogItemIdsCount; CatalogItemIndex++)
+				{
+					const EOS_Ecom_CatalogItemId& CatalogItemId = SandboxOwnershipEntry.OwnedCatalogItemIds[CatalogItemIndex];
+
+					FPurchaseReceipt& PurchaseReceipt = CachedReceipts.Emplace_GetRef();
+					PurchaseReceipt.TransactionState = EPurchaseTransactionState::Purchased;
+
+					UE_LOG_ONLINE(VeryVerbose, TEXT("[FOnlinePurchaseEOS::QueryReceipts] Catalog Item %d"), CatalogItemIndex);
+					UE_LOG_ONLINE(VeryVerbose, TEXT("[FOnlinePurchaseEOS::QueryReceipts] Catalog Item Id: %hs"), CatalogItemId);
+
+					// The Receipt structure is 1 FReceiptOfferEntry with 1 line item for each EOS_Ecom_CatalogItem the player owns
+					// Receipts do not contain any entitlements - use QueryEntitlements in the Online Entitlements interface to retrieve entitlement data
+					
+					// OfferId is set to an empty string as the OfferId is not retrievable via EOS SDK QueryOwnership APIs
+					PurchaseReceipt.AddReceiptOffer(FOfferNamespace(), "", 1);
+					FPurchaseReceipt::FLineItemInfo& LineItem = PurchaseReceipt.ReceiptOffers[0].LineItems.Emplace_GetRef();
+					// The item that is owned is the CatalogItemId
+					LineItem.UniqueId = StringCast<TCHAR>(CatalogItemId).Get();
+
+					EOS_Ecom_CatalogItem* CatalogItem = nullptr;
+					EOS_Ecom_CopyItemByIdOptions CopyCatalogItemOptions = {};
+					CopyCatalogItemOptions.ApiVersion = 1;
+					UE_EOS_CHECK_API_MISMATCH(EOS_ECOM_COPYITEMBYID_API_LATEST, 1);
+					CopyCatalogItemOptions.ItemId = CatalogItemId;
+					CopyCatalogItemOptions.LocalUserId = Data->LocalUserId;
+					EOS_EResult CopyCatalogItemResult = EOS_Ecom_CopyItemById(EOSSubsystem->EcomHandle, &CopyCatalogItemOptions, &CatalogItem);
+
+					if (CopyCatalogItemResult == EOS_EResult::EOS_Success || CopyCatalogItemResult == EOS_EResult::EOS_Ecom_CatalogItemStale)
+					{
+						LineItem.ItemName = CatalogItem->TitleText ? StringCast<TCHAR>(CatalogItem->TitleText).Get() : FString();
+						EOS_Ecom_CatalogItem_Release(CatalogItem);
+					}
+					else
+					{
+						// EOS_Ecom_CopyItemById will return EOS_NotFound is QueryOffers has NOT been called OR if the catalogItemId doesn't tie back to a specific offer (bundles, season packs, etc...)
+						UE_LOG_ONLINE(Verbose, TEXT("[FOnlinePurchaseEOS::QueryReceipts] EOS_Ecom_CopyItemById: failed with error (%s)"), *LexToString(CopyCatalogItemResult));
+					}
+				}
+			}
+			OnComplete.ExecuteIfBound(ONLINE_ERROR(EOnlineErrorResult::Success));
+		}
+		else
+		{
+			UE_LOG_ONLINE(Error, TEXT("[FOnlinePurchaseEOS::QueryReceipts] EOS_Ecom_QueryOwnership: failed with error (%s)"), *LexToString(Data->ResultCode));
+			OnComplete.ExecuteIfBound(ONLINE_ERROR(EOnlineErrorResult::Unknown));
+			return;
+		}
+	};
+
+	EOS_Ecom_QueryOwnershipBySandboxIds(EOSSubsystem->EcomHandle, &Options, CallbackObj, CallbackObj->GetCallbackPtr());
+}
+
+void FOnlinePurchaseEOS::GetReceipts(const FUniqueNetId& UserId, TArray<FPurchaseReceipt>& OutReceipts) const
+{
+	OutReceipts.Empty(); 
+
+	const FUniqueNetIdEOS& UserEOSId = FUniqueNetIdEOS::Cast(UserId);
+	const EOS_EpicAccountId AccountId = UserEOSId.GetEpicAccountId();
+	if (AccountId == nullptr)
+	{
+		UE_LOG_ONLINE(Error, TEXT("[%hs] EpicAccountId is null"), __FUNCTION__);
+		return;
+	}
+
+	if (CachedReceipts.IsEmpty())
+	{
+		UE_LOG_ONLINE(Error, TEXT("[%hs] CachedReceipts empty, call QueryReceipts"), __FUNCTION__);
+		return;
+	}
+
+	OutReceipts = CachedReceipts;
+}
+
+typedef TEOSCallback<EOS_Ecom_OnRedeemEntitlementsCallback, EOS_Ecom_RedeemEntitlementsCallbackInfo, FOnlinePurchaseEOS> FPurchaseRedeemReceiptCallback;
+
+void FOnlinePurchaseEOS::FinalizeReceiptValidationInfo(const FUniqueNetId& UserId, FString& InReceiptValidationInfo, const FOnFinalizeReceiptValidationInfoComplete& Delegate)
+{
+	const FUniqueNetIdEOS& UserEOSId = FUniqueNetIdEOS::Cast(UserId);
+	const EOS_EpicAccountId AccountId = UserEOSId.GetEpicAccountId();
+	if (AccountId == nullptr)
+	{
+		UE_LOG_ONLINE(Error, TEXT("[%hs] EpicAccountId is null"), __FUNCTION__);
+		Delegate.ExecuteIfBound(ONLINE_ERROR(EOnlineErrorResult::InvalidUser), InReceiptValidationInfo);
+		return;
+	}
+	if (InReceiptValidationInfo.IsEmpty())
+	{
+		UE_LOG_ONLINE(Error, TEXT("[%hs] InReceiptValidationInfo is empty"), __FUNCTION__);
+		Delegate.ExecuteIfBound(ONLINE_ERROR(EOnlineErrorResult::InvalidParams), InReceiptValidationInfo);
+		return;
+	}
+
+	char const* Ids[1];
+	FTCHARToUTF8 EntitlementId(*InReceiptValidationInfo);
+	Ids[0] = EntitlementId.Get();
+
+	EOS_Ecom_RedeemEntitlementsOptions Options = { };
+	Options.ApiVersion = 2;
+	UE_EOS_CHECK_API_MISMATCH(EOS_ECOM_REDEEMENTITLEMENTS_API_LATEST, 2);
+	Options.LocalUserId = AccountId;
+	Options.EntitlementIdCount = 1;
+	Options.EntitlementIds = Ids;
+
+	FPurchaseRedeemReceiptCallback* CallbackObj = new FPurchaseRedeemReceiptCallback(FOnlinePurchaseEOSWeakPtr(AsShared()));
+	CallbackObj->CallbackLambda = [this, Info = FString(InReceiptValidationInfo), OnComplete = FOnFinalizeReceiptValidationInfoComplete(Delegate)](const EOS_Ecom_RedeemEntitlementsCallbackInfo* Data)
+		{
+			EOS_EResult Result = Data->ResultCode;
+			if (Result == EOS_EResult::EOS_Success)
+			{
+
+				// Find the receipt in our list and mark as redeemed (clear the validation info) - It's possible the entitlement is NOT in our cached receipts
+				for (FPurchaseReceipt& SearchReceipt : CachedReceipts)
+				{
+					if (!SearchReceipt.ReceiptOffers.IsEmpty())
+					{
+						if (!SearchReceipt.ReceiptOffers[0].LineItems.IsEmpty())
+						{
+							if (SearchReceipt.ReceiptOffers[0].LineItems[0].ValidationInfo == Info)
+							{
+								// Clearing this field tells the game it can't be redeemed - ConsumedCount on FOnlineEntitlement should be used to determine if an entitlement can be redeemed
+								SearchReceipt.ReceiptOffers[0].LineItems[0].ValidationInfo.Empty();
+								break;
+							}
+						}
+					}
+				}
+
+				OnComplete.ExecuteIfBound(ONLINE_ERROR(EOnlineErrorResult::Success), Info);
+			}
+			else
+			{
+				UE_LOG_ONLINE(Error, TEXT("[FOnlinePurchaseEOS::FinalizeReceiptValidationInfo] EOS_Ecom_RedeemEntitlements: failed with error (%s)"), *LexToString(Data->ResultCode));
+				OnComplete.ExecuteIfBound(ONLINE_ERROR(EOnlineErrorResult::Unknown), Info);
+				return;
+			}
+		};
+	EOS_Ecom_RedeemEntitlements(EOSSubsystem->EcomHandle, &Options, CallbackObj, CallbackObj->GetCallbackPtr());
+}
+
+bool FOnlinePurchaseEOS::HandlePurchaseExec(UWorld* InWorld, const TCHAR* Cmd, FOutputDevice& Ar)
+{
+	if (!EOSSubsystem || !EOSSubsystem->UserManager.IsValid() || !EOSSubsystem->PurchaseInterfacePtr.IsValid())
+	{
+		UE_LOG_ONLINE(Error, TEXT("[%hs] EOSSubsystem or UserManager or PurchaseInterface invalid"), __FUNCTION__);
+		return false;
+	}
+
+	if (FParse::Command(&Cmd, TEXT("RECEIPTS")))
+	{
+		QueryReceipts(*EOSSubsystem->UserManager->GetLocalUniqueNetIdEOS(EOSSubsystem->UserManager->GetDefaultLocalUser()), false,
+			FOnQueryReceiptsComplete::CreateLambda([this](const FOnlineError& Result)
+				{
+					if (Result.WasSuccessful())
+					{
+						UE_LOG_ONLINE(Log, TEXT("[FOnlinePurchaseEOS::HandlePurchaseExec] QueryReceipts success"));
+						for (const FPurchaseReceipt& Receipt : CachedReceipts)
+						{
+							if (!Receipt.ReceiptOffers.IsEmpty())
+							{
+								if (!Receipt.ReceiptOffers[0].LineItems.IsEmpty())
+								{
+									UE_LOG_ONLINE(Log, TEXT("[FOnlinePurchaseEOS::HandlePurchaseExec] Player owns ItemId: %s"), *Receipt.ReceiptOffers[0].LineItems[0].UniqueId);
+								}
+							}
+						}
+					}
+					else
+					{
+						UE_LOG_ONLINE(Error, TEXT("[FOnlinePurchaseEOS::HandlePurchaseExec] QueryReceipts error (%s)"), *Result.GetErrorRaw());
+					}
+				}));
+		return true;
+	}
+	else if (FParse::Command(&Cmd, TEXT("FINALIZE"))) // ONLINE PURCHASES FINALIZE EntitlementId=<ENTITLEMENTID>
+	{
+		FString EntitlementId;
+		FParse::Value(Cmd, TEXT("EntitlementId="), EntitlementId);
+
+		FinalizeReceiptValidationInfo(*EOSSubsystem->UserManager->GetLocalUniqueNetIdEOS(EOSSubsystem->UserManager->GetDefaultLocalUser()), EntitlementId,
+			FOnFinalizeReceiptValidationInfoComplete::CreateLambda([this](const FOnlineError& Result, const FString& Info)
+				{
+					if (Result.WasSuccessful())
+					{
+						UE_LOG_ONLINE(Log, TEXT("[FOnlinePurchaseEOS::HandlePurchaseExec] FinalizeReceiptValidationInfo success - info (%s)"), *Info);
+					}
+					else
+					{
+						UE_LOG_ONLINE(Error, TEXT("[FOnlinePurchaseEOS::HandlePurchaseExec] FinalizeReceiptValidationInfo error (%s)"), *Result.GetErrorRaw());
+					}
+				}));
+		return true;
+	}
+	else
+	{
+		return false;
+	}
+}
+
+#endif
\ No newline at end of file
diff --git a/Engine/Plugins/Online/OnlineSubsystemEOS/Source/OnlineSubsystemEOS/Private/OnlinePurchaseEOS.h b/Engine/Plugins/Online/OnlineSubsystemEOS/Source/OnlineSubsystemEOS/Private/OnlinePurchaseEOS.h
new file mode 100644
--- /dev/null
+++ b/Engine/Plugins/Online/OnlineSubsystemEOS/Source/OnlineSubsystemEOS/Private/OnlinePurchaseEOS.h
@@ -0,0 +1,66 @@
+// Copyright Epic Games, Inc. All Rights Reserved.
+
+#pragma once
+
+#include "CoreMinimal.h"
+#include "Online/CoreOnline.h"
+#include "Interfaces/OnlinePurchaseInterface.h"
+#include "OnlineSubsystemEOSTypes.h"
+
+#if WITH_EOS_SDK
+#include "eos_ecom_types.h"
+
+class UWorld;
+
+/**
+ * Checkout offers, retrieve receipts (owned EOS_Ecom_CatalogItemIds) and finalize receipt validation (redeem entitlement)
+ */
+
+class FOnlinePurchaseEOS :
+	public IOnlinePurchase,
+	public TSharedFromThis<FOnlinePurchaseEOS, ESPMode::ThreadSafe>
+{
+public:
+	virtual ~FOnlinePurchaseEOS() = default;
+	
+	virtual bool IsAllowedToPurchase(const FUniqueNetId& UserId) override { return true; }
+	virtual bool IsAllowedToPurchase(const FUniqueNetId& UserId, FOnlineError& Error) override
+	{
+		Error = FOnlineError(EOnlineErrorResult::Success);
+		return true;
+	}
+	/** Checkout store offers and return a receipt. Receipt contains the EOS_Ecom_EntitlementId in the LineItem.ValidationInfo which can be redeemed via FinalizeReceiptValidationInfo */
+	virtual void Checkout(const FUniqueNetId& UserId, const FPurchaseCheckoutRequest& CheckoutRequest, const FOnPurchaseCheckoutComplete& Delegate) override;
+	/** Checkout store offers */
+	virtual void Checkout(const FUniqueNetId& UserId, const FPurchaseCheckoutRequest& CheckoutRequest, const FOnPurchaseReceiptlessCheckoutComplete& Delegate) override;
+	/** Not implemented */
+	virtual void FinalizePurchase(const FUniqueNetId& UserId, const FString& ReceiptId) override;
+	/** Not implemented */
+	virtual void RedeemCode(const FUniqueNetId& UserId, const FRedeemCodeRequest& RedeemCodeRequest, const FOnPurchaseRedeemCodeComplete& Delegate) override;
+	/* Retrieve and cache ownership data only - To retrieve entitlements use QueryEntitlements in the OnlineEntitlements interface */
+	virtual void QueryReceipts(const FUniqueNetId& UserId, bool /* bRestoreReceipts */, const FOnQueryReceiptsComplete& Delegate) override;
+	/* Retrieve cached ownership data */
+	virtual void GetReceipts(const FUniqueNetId& UserId, TArray<FPurchaseReceipt>& OutReceipts) const override;
+	/** Used to redeem an entitlement - Pass in the EntitlementId in InReceiptValidationInfo - The EntitlementId is provided in the CheckoutCallback with receipt in LineItem.ValidationInfo and is provided in GetAllEntitlements, GetEntitlement, and GetItemEntitlement in the OnlineEntitlements interface */
+	virtual void FinalizeReceiptValidationInfo(const FUniqueNetId& UserId, FString& InReceiptValidationInfo, const FOnFinalizeReceiptValidationInfoComplete& Delegate) override;
+	
+	FOnlinePurchaseEOS(FOnlineSubsystemEOS* InSubsystem);
+	
+	bool HandlePurchaseExec(UWorld* InWorld, const TCHAR* Cmd, FOutputDevice& Ar);
+
+private:
+	/** Default constructor disabled */
+	FOnlinePurchaseEOS() = delete;
+
+	/** Reference to the main EOS subsystem */
+	FOnlineSubsystemEOS* EOSSubsystem;
+
+	/** List of receipts for the user */
+	//@todo joeg - make this support multiple users
+	TArray<FPurchaseReceipt> CachedReceipts;
+};
+
+typedef TSharedPtr<FOnlinePurchaseEOS, ESPMode::ThreadSafe> FOnlinePurchaseEOSPtr;
+typedef TWeakPtr<FOnlinePurchaseEOS, ESPMode::ThreadSafe> FOnlinePurchaseEOSWeakPtr;
+
+#endif
diff --git a/Engine/Plugins/Online/OnlineSubsystemEOS/Source/OnlineSubsystemEOS/Private/OnlineStoreEOSV2.cpp b/Engine/Plugins/Online/OnlineSubsystemEOS/Source/OnlineSubsystemEOS/Private/OnlineStoreEOSV2.cpp
new file mode 100644
--- /dev/null
+++ b/Engine/Plugins/Online/OnlineSubsystemEOS/Source/OnlineSubsystemEOS/Private/OnlineStoreEOSV2.cpp
@@ -0,0 +1,201 @@
+// Copyright Epic Games, Inc. All Rights Reserved.
+
+#include "OnlineStoreEOSV2.h"
+#include "EOSSettings.h"
+#include "IEOSSDKManager.h"
+
+#if WITH_EOS_SDK
+
+#include "OnlineSubsystemEOS.h"
+#include "OnlineSubsystemEOSPrivate.h"
+#include "UserManagerEOS.h"
+#include "eos_ecom.h"
+
+#define ONLINE_ERROR_NAMESPACE "com.epicgames.oss.eos.error"
+
+
+FOnlineStoreEOSV2::FOnlineStoreEOSV2(FOnlineSubsystemEOS* InSubsystem)
+	: EOSSubsystem(InSubsystem)
+{
+	check(EOSSubsystem != nullptr);
+}
+
+void FOnlineStoreEOSV2::QueryCategories(const FUniqueNetId& UserId, const FOnQueryOnlineStoreCategoriesComplete& Delegate)
+{
+	UE_LOG_ONLINE(Verbose, TEXT("[%hs] Not implemented"), __FUNCTION__);
+	Delegate.ExecuteIfBound(false, "Not implemented");
+}
+
+void FOnlineStoreEOSV2::GetCategories(TArray<FOnlineStoreCategory>& OutCategories) const
+{
+	UE_LOG_ONLINE(Verbose, TEXT("[%hs] Not implemented"), __FUNCTION__);
+	OutCategories.Reset();
+}
+
+void FOnlineStoreEOSV2::QueryOffersByFilter(const FUniqueNetId& UserId, const FOnlineStoreFilter& Filter, const FOnQueryOnlineStoreOffersComplete& Delegate)
+{
+	UE_LOG_ONLINE(Verbose, TEXT("[%hs] Not implemented - No EOS SDK API to query offers by filter - Querying all offers"), __FUNCTION__);
+	QueryOffers(UserId, Delegate);
+}
+
+void FOnlineStoreEOSV2::QueryOffersById(const FUniqueNetId& UserId, const TArray<FUniqueOfferId>& OfferIds, const FOnQueryOnlineStoreOffersComplete& Delegate)
+{
+	UE_LOG_ONLINE(Verbose, TEXT("[%hs] Not implemented - No EOS SDK API to query offers by Id - Querying all offers"), __FUNCTION__);
+	QueryOffers(UserId, Delegate);
+}
+
+typedef TEOSCallback<EOS_Ecom_OnQueryOffersCallback, EOS_Ecom_QueryOffersCallbackInfo, FOnlineStoreEOSV2> FOnlineStoreV2QueryOffersCallback;
+
+void FOnlineStoreEOSV2::QueryOffers(const FUniqueNetId& UserId, const FOnQueryOnlineStoreOffersComplete& Delegate)
+{
+	const FUniqueNetIdEOS& UserEOSId = FUniqueNetIdEOS::Cast(UserId);
+	const EOS_EpicAccountId AccountId = UserEOSId.GetEpicAccountId();
+	if (AccountId == nullptr)
+	{
+		UE_LOG_ONLINE(Error, TEXT("[%hs] EpicAccountId is null"), __FUNCTION__);
+		Delegate.ExecuteIfBound(false, TArray<FUniqueOfferId>(), "EpicAccountId is null");
+		return;
+	}
+
+	EOS_Ecom_QueryOffersOptions Options = { };
+	Options.ApiVersion = 1;
+	UE_EOS_CHECK_API_MISMATCH(EOS_ECOM_QUERYOFFERS_API_LATEST, 1);
+	Options.LocalUserId = AccountId;
+
+	FOnlineStoreV2QueryOffersCallback* CallbackObj = new FOnlineStoreV2QueryOffersCallback(FOnlineStoreEOSV2WeakPtr(AsShared()));
+	CallbackObj->CallbackLambda = [this, OnComplete = FOnQueryOnlineStoreOffersComplete(Delegate)](const EOS_Ecom_QueryOffersCallbackInfo* Data)
+		{
+			CachedOfferIds.Reset();
+			CachedOffers.Reset();
+
+			if (Data->ResultCode == EOS_EResult::EOS_Success)
+			{
+				EOS_Ecom_GetOfferCountOptions CountOptions = { };
+				CountOptions.ApiVersion = 1;
+				UE_EOS_CHECK_API_MISMATCH(EOS_ECOM_GETOFFERCOUNT_API_LATEST, 1);
+				CountOptions.LocalUserId = Data->LocalUserId;
+				uint32 OfferCount = EOS_Ecom_GetOfferCount(EOSSubsystem->EcomHandle, &CountOptions);
+
+				EOS_Ecom_CopyOfferByIndexOptions OfferOptions = { };
+				OfferOptions.ApiVersion = 3;
+				UE_EOS_CHECK_API_MISMATCH(EOS_ECOM_COPYOFFERBYINDEX_API_LATEST, 3);
+				OfferOptions.LocalUserId = Data->LocalUserId;
+
+				for (uint32 OfferIndex = 0; OfferIndex < OfferCount; OfferIndex++)
+				{
+					EOS_Ecom_CatalogOffer* Offer = nullptr;
+					OfferOptions.OfferIndex = OfferIndex;
+					EOS_EResult OfferResult = EOS_Ecom_CopyOfferByIndex(EOSSubsystem->EcomHandle, &OfferOptions, &Offer);
+					if (OfferResult == EOS_EResult::EOS_Success)
+					{
+						FOnlineStoreOfferRef OfferRef = MakeShared<FOnlineStoreOffer>();
+						//OfferId represents an EOS_Ecom_CatalogOfferId
+						OfferRef->OfferId = Offer->Id ? FUniqueOfferId(StringCast<TCHAR>(Offer->Id).Get()) : FString();
+						OfferRef->Title = FText::FromString(Offer->TitleText ? StringCast<TCHAR>(Offer->TitleText).Get() : FString());
+						OfferRef->Description = FText::FromString(Offer->DescriptionText ? StringCast<TCHAR>(Offer->DescriptionText).Get() : FString());
+						OfferRef->LongDescription = FText::FromString(Offer->LongDescriptionText ? StringCast<TCHAR>(Offer->LongDescriptionText).Get() : FString());
+						OfferRef->ExpirationDate = FDateTime::FromUnixTimestamp(Offer->ExpirationTimestamp);
+						OfferRef->CurrencyCode = Offer->CurrencyCode ? StringCast<TCHAR>(Offer->CurrencyCode).Get() : FString();
+						OfferRef->DynamicFields.Add(TEXT("bAvailableForPurchase"), LexToString(Offer->bAvailableForPurchase));
+						OfferRef->DynamicFields.Add(TEXT("PurchaseLimit"), LexToString(Offer->PurchaseLimit));
+
+						if (Offer->PriceResult == EOS_EResult::EOS_Success)
+						{
+#if EOS_ECOM_CATALOGOFFER_API_LATEST >= 3
+							OfferRef->RegularPrice = Offer->OriginalPrice64;
+							OfferRef->NumericPrice = Offer->CurrentPrice64;
+#else
+							OfferRef->RegularPrice = Offer->OriginalPrice;
+							OfferRef->NumericPrice = Offer->CurrentPrice;
+#endif
+							OfferRef->DiscountType = Offer->DiscountPercentage == 0 ? EOnlineStoreOfferDiscountType::NotOnSale : EOnlineStoreOfferDiscountType::DiscountAmount;
+							OfferRef->RegularPriceText = FText::AsCurrencyBase(OfferRef->RegularPrice, OfferRef->CurrencyCode, NULL, Offer->DecimalPoint);
+							OfferRef->PriceText = FText::AsCurrencyBase(OfferRef->NumericPrice, OfferRef->CurrencyCode, NULL, Offer->DecimalPoint);
+						}
+						else
+						{
+							UE_LOG_ONLINE(Verbose, TEXT("[FOnlineStoreEOSV2::QueryOffers] Price Result %s Offer %s"), *LexToString(Offer->PriceResult), *OfferRef->OfferId);
+						}
+
+						CachedOffers.Add(OfferRef);
+						CachedOfferIds.Add(OfferRef->OfferId);
+
+						EOS_Ecom_CatalogOffer_Release(Offer);
+					}
+					else
+					{
+						UE_LOG_ONLINE(Warning, TEXT("[FOnlineStoreEOSV2::QueryOffers] EOS_Ecom_CopyOfferByIndex result code (%s)"), *LexToString(OfferResult));
+						continue;
+					}
+				}
+				OnComplete.ExecuteIfBound(true, CachedOfferIds, TEXT(""));
+			}
+			else
+			{
+				FString ErrorString = FString::Printf(TEXT("result code (%s)"), *LexToString(Data->ResultCode));
+				UE_LOG_ONLINE(Error, TEXT("[FOnlineStoreEOSV2::QueryOffers] %s"), *ErrorString);
+				OnComplete.ExecuteIfBound(false, CachedOfferIds, ErrorString);
+			}
+		};
+	EOS_Ecom_QueryOffers(EOSSubsystem->EcomHandle, &Options, CallbackObj, CallbackObj->GetCallbackPtr());
+}
+
+void FOnlineStoreEOSV2::GetOffers(TArray<FOnlineStoreOfferRef>& OutOffers) const
+{
+	OutOffers.Empty();
+
+	if (CachedOffers.IsEmpty())
+	{
+		UE_LOG_ONLINE(Verbose, TEXT("[%hs] CachedOffers empty, call QueryOffers"), __FUNCTION__);
+		return;
+	}
+
+	OutOffers = CachedOffers;
+}
+
+TSharedPtr<FOnlineStoreOffer> FOnlineStoreEOSV2::GetOffer(const FUniqueOfferId& OfferId) const
+{
+	for (FOnlineStoreOfferRef Offer : CachedOffers)
+	{
+		if (Offer->OfferId == OfferId)
+		{
+			return Offer;
+		}
+	}
+	
+	UE_LOG_ONLINE(Verbose, TEXT("[%hs] CachedOffers does not contain an offer with OfferId (%s)"), __FUNCTION__,*OfferId);
+	return nullptr;
+}
+
+bool FOnlineStoreEOSV2::HandleEcomExec(UWorld* InWorld, const TCHAR* Cmd, FOutputDevice& Ar)
+{
+	if (!EOSSubsystem || !EOSSubsystem->UserManager.IsValid() || !EOSSubsystem->StoreInterfaceV2Ptr.IsValid())
+	{
+		UE_LOG_ONLINE(Error, TEXT("[%hs] EOSSubsystem or UserManager or StoreInterfaceV2 invalid"), __FUNCTION__);
+		return false;
+	}
+
+	if (FParse::Command(&Cmd, TEXT("OFFERS")))
+	{
+		QueryOffers(*EOSSubsystem->UserManager->GetLocalUniqueNetIdEOS(EOSSubsystem->UserManager->GetDefaultLocalUser()),
+			FOnQueryOnlineStoreOffersComplete::CreateLambda([this](bool bWasSuccessful, const TArray<FUniqueOfferId>& OfferIds, const FString& ErrorStr)
+				{
+					if (bWasSuccessful)
+					{
+						UE_LOG_ONLINE(Log, TEXT("[FOnlineStoreEOSV2::HandleEcomExec] QueryOffers success"));
+						for (const FUniqueOfferId& OfferId : OfferIds)
+						{
+							UE_LOG_ONLINE(Log, TEXT("[FOnlineStoreEOSV2::HandleEcomExec] OfferId: %s"), *OfferId);
+						}
+					}
+					else
+					{
+						UE_LOG_ONLINE(Error, TEXT("[FOnlineStoreEOSV2::HandleEcomExec] QueryOffers error (%s)"), *ErrorStr);
+					}
+				}));
+		return true;
+	}
+	return false;
+}
+
+#endif
\ No newline at end of file
diff --git a/Engine/Plugins/Online/OnlineSubsystemEOS/Source/OnlineSubsystemEOS/Private/OnlineStoreEOSV2.h b/Engine/Plugins/Online/OnlineSubsystemEOS/Source/OnlineSubsystemEOS/Private/OnlineStoreEOSV2.h
new file mode 100644
--- /dev/null
+++ b/Engine/Plugins/Online/OnlineSubsystemEOS/Source/OnlineSubsystemEOS/Private/OnlineStoreEOSV2.h
@@ -0,0 +1,55 @@
+// Copyright Epic Games, Inc. All Rights Reserved.
+
+#pragma once
+
+#include "CoreMinimal.h"
+#include "Online/CoreOnline.h"
+#include "Interfaces/OnlineStoreInterfaceV2.h"
+#include "OnlineSubsystemEOSTypes.h"
+
+#if WITH_EOS_SDK
+#include "eos_ecom_types.h"
+
+class UWorld;
+
+/**
+ * Retrieve offers from store. Offers can be checked out using the OnlinePurchase interface
+ */
+class FOnlineStoreEOSV2 :
+	public IOnlineStoreV2,
+	public TSharedFromThis<FOnlineStoreEOSV2, ESPMode::ThreadSafe>
+{
+public:
+	virtual ~FOnlineStoreEOSV2() = default;
+	
+	virtual void QueryCategories(const FUniqueNetId& UserId, const FOnQueryOnlineStoreCategoriesComplete& Delegate) override;
+	virtual void GetCategories(TArray<FOnlineStoreCategory>& OutCategories) const override;
+	virtual void QueryOffersByFilter(const FUniqueNetId& UserId, const FOnlineStoreFilter& Filter, const FOnQueryOnlineStoreOffersComplete& Delegate) override;
+	virtual void QueryOffersById(const FUniqueNetId& UserId, const TArray<FUniqueOfferId>& OfferIds, const FOnQueryOnlineStoreOffersComplete& Delegate) override;
+	virtual void GetOffers(TArray<FOnlineStoreOfferRef>& OutOffers) const override;
+	virtual TSharedPtr<FOnlineStoreOffer> GetOffer(const FUniqueOfferId& OfferId) const override;
+	
+	FOnlineStoreEOSV2(FOnlineSubsystemEOS* InSubsystem);
+	
+	bool HandleEcomExec(UWorld* InWorld, const TCHAR* Cmd, FOutputDevice& Ar);
+
+private:
+	/** Default constructor disabled */
+	FOnlineStoreEOSV2() = delete;
+
+	/** Retrieve Epic Games Store offers */
+	void QueryOffers(const FUniqueNetId& UserId, const FOnQueryOnlineStoreOffersComplete& Delegate);
+
+	/** Reference to the main EOS subsystem */
+	FOnlineSubsystemEOS* EOSSubsystem;
+
+	/** The set of offers for this title */
+	TArray<FOnlineStoreOfferRef> CachedOffers;
+	/** List of offer ids for this title */
+	TArray<FUniqueOfferId> CachedOfferIds;
+};
+
+typedef TSharedPtr<FOnlineStoreEOSV2, ESPMode::ThreadSafe> FOnlineStoreEOSV2Ptr;
+typedef TWeakPtr<FOnlineStoreEOSV2, ESPMode::ThreadSafe> FOnlineStoreEOSV2WeakPtr;
+
+#endif
diff --git a/Engine/Plugins/Online/OnlineSubsystemEOS/Source/OnlineSubsystemEOS/Private/OnlineSubsystemEOS.cpp b/Engine/Plugins/Online/OnlineSubsystemEOS/Source/OnlineSubsystemEOS/Private/OnlineSubsystemEOS.cpp
--- a/Engine/Plugins/Online/OnlineSubsystemEOS/Source/OnlineSubsystemEOS/Private/OnlineSubsystemEOS.cpp
+++ b/Engine/Plugins/Online/OnlineSubsystemEOS/Source/OnlineSubsystemEOS/Private/OnlineSubsystemEOS.cpp
@@ -13,6 +13,9 @@
 #include "OnlineStoreEOS.h"
 #include "OnlinePlayerSanctionEOS.h"
 #include "OnlinePlayerReportEOS.h"
+#include "OnlineEntitlementsEOS.h"
+#include "OnlinePurchaseEOS.h"
+#include "OnlineStoreEOSV2.h"
 #include "EOSSettings.h"
 #include "EOSShared.h"
 #include "EOSVoiceChat.h"
@@ -464,10 +467,22 @@ bool FOnlineSubsystemEOS::Init()
 	StatsInterfacePtr = MakeShared<FOnlineStatsEOS>(this);
 	LeaderboardsInterfacePtr = MakeShared<FOnlineLeaderboardsEOS>(this);
 	AchievementsInterfacePtr = MakeShared<FOnlineAchievementsEOS>(this);
-	StoreInterfacePtr = MakeShared<FOnlineStoreEOS>(this);
 	PlayerSanctionEOSPtr = MakeShared<FOnlinePlayerSanctionEOS>(this);
 	PlayerReportInterfacePtr = MakeShared<FOnlinePlayerReportEOS>(this);
 
+	FEOSSettings Settings = UEOSSettings::GetSettings();
+	if (Settings.bUseNewEcomFlow)
+	{
+		EntitlementsInterfacePtr = MakeShared<FOnlineEntitlementsEOS>(this);
+		PurchaseInterfacePtr = MakeShared<FOnlinePurchaseEOS>(this);
+		StoreInterfaceV2Ptr = MakeShared<FOnlineStoreEOSV2>(this);
+	}
+	else
+	{
+		StoreInterfacePtr = MakeShared<FOnlineStoreEOS>(this);
+		UE_LOG_ONLINE(Warning, TEXT("[FOnlineSubsystemEOS::Init] bUseNewEcomFlow set to false. Note the old ecom flow will be deprecated in a future version"));
+	}
+
 	// We initialized ok so we can tick
 	StartTicker();
 
@@ -498,9 +513,12 @@ bool FOnlineSubsystemEOS::Shutdown()
 	StatsInterfacePtr = nullptr;
 	LeaderboardsInterfacePtr = nullptr;
 	AchievementsInterfacePtr = nullptr;
-	StoreInterfacePtr = nullptr;
 	TitleFileInterfacePtr = nullptr;
 	UserCloudInterfacePtr = nullptr;
+	EntitlementsInterfacePtr = nullptr;
+	PurchaseInterfacePtr = nullptr;
+	StoreInterfaceV2Ptr = nullptr;
+	StoreInterfacePtr = nullptr;
 
 #if WITH_EOSVOICECHAT
 	for (TPair<FUniqueNetIdRef, FOnlineSubsystemEOSVoiceChatUserWrapperRef>& Pair : LocalVoiceChatUsers)
@@ -542,10 +560,22 @@ bool FOnlineSubsystemEOS::Exec(UWorld* InWorld, const TCHAR* Cmd, FOutputDevice&
 	{
 		bWasHandled = UserManager->HandleFriendsExec(InWorld, Cmd, Ar);
 	}
-	else if (StoreInterfacePtr != nullptr && FParse::Command(&Cmd, TEXT("ECOM"))) /* ONLINE ECOM ... */
+	else if (StoreInterfacePtr != nullptr && FParse::Command(&Cmd, TEXT("ECOM"))) /* ONLINE ECOM ... only ran if bUseNewEcomFlow is false */
 	{
 		bWasHandled = StoreInterfacePtr->HandleEcomExec(InWorld, Cmd, Ar);
 	}
+	else if (StoreInterfaceV2Ptr != nullptr && FParse::Command(&Cmd, TEXT("ECOM"))) /* ONLINE ECOM ... only ran if bUseNewEcomFlow is true */
+	{
+		bWasHandled = StoreInterfaceV2Ptr->HandleEcomExec(InWorld, Cmd, Ar);
+	}
+	else if (PurchaseInterfacePtr != nullptr && FParse::Command(&Cmd, TEXT("PURCHASE"))) /* ONLINE PURCHASE ... only ran if bUseNewEcomFlow is true */
+	{
+		bWasHandled = PurchaseInterfacePtr->HandlePurchaseExec(InWorld, Cmd, Ar);
+	}
+	else if (EntitlementsInterfacePtr != nullptr && FParse::Command(&Cmd, TEXT("ENTITLEMENTS"))) /* ONLINE ENTITLEMENTS ... only ran if bUseNewEcomFlow is true */
+	{
+		bWasHandled = EntitlementsInterfacePtr->HandleEntitlementsExec(InWorld, Cmd, Ar);
+	}
 	else if (LeaderboardsInterfacePtr != nullptr && FParse::Command(&Cmd, TEXT("LEADERBOARDS"))) /* ONLINE LEADERBOARDS ... */
 	{
 		bWasHandled = LeaderboardsInterfacePtr->HandleLeaderboardsExec(InWorld, Cmd, Ar);
@@ -635,6 +665,9 @@ FOnlineSubsystemEOS::FOnlineSubsystemEOS(FName InInstanceName) :
 	, StoreInterfacePtr(nullptr)
 	, TitleFileInterfacePtr(nullptr)
 	, UserCloudInterfacePtr(nullptr)
+	, EntitlementsInterfacePtr(nullptr)
+	, PurchaseInterfacePtr(nullptr)
+	, StoreInterfaceV2Ptr(nullptr)
 {
 	StopTicker();
 }
@@ -662,9 +695,20 @@ IOnlineUserCloudPtr FOnlineSubsystemEOS::GetUserCloudInterface() const
 
 IOnlineEntitlementsPtr FOnlineSubsystemEOS::GetEntitlementsInterface() const
 {
-	UE_LOG_ONLINE(Error, TEXT("Entitlements Interface Requested"));
-	return nullptr;
-};
+	FEOSSettings Settings = UEOSSettings::GetSettings();
+	if (Settings.bUseNewEcomFlow)
+	{
+		return EntitlementsInterfacePtr;
+	}
+	else
+	{
+		UE_CALL_ONCE([]()
+		{
+			UE_LOG_ONLINE(Warning, TEXT("Entitlements Interface Requested - only supported with bUseNewEcomFlow"));
+		});
+		return nullptr;
+	}
+}
 
 IOnlineLeaderboardsPtr FOnlineSubsystemEOS::GetLeaderboardsInterface() const
 {
@@ -693,12 +737,28 @@ IOnlineTitleFilePtr FOnlineSubsystemEOS::GetTitleFileInterface() const
 
 IOnlineStoreV2Ptr FOnlineSubsystemEOS::GetStoreV2Interface() const
 {
-	return StoreInterfacePtr;
+	FEOSSettings Settings = UEOSSettings::GetSettings();
+	if (Settings.bUseNewEcomFlow)
+	{
+		return StoreInterfaceV2Ptr;
+	}
+	else
+	{
+		return StoreInterfacePtr;
+	}
 }
 
 IOnlinePurchasePtr FOnlineSubsystemEOS::GetPurchaseInterface() const
 {
-	return StoreInterfacePtr;
+	FEOSSettings Settings = UEOSSettings::GetSettings();
+	if (Settings.bUseNewEcomFlow)
+	{
+		return PurchaseInterfacePtr;
+	}
+	else
+	{
+		return StoreInterfacePtr;
+	}
 }
 
 IOnlineAchievementsPtr FOnlineSubsystemEOS::GetAchievementsInterface() const
diff --git a/Engine/Plugins/Online/OnlineSubsystemEOS/Source/OnlineSubsystemEOS/Private/OnlineSubsystemEOS.h b/Engine/Plugins/Online/OnlineSubsystemEOS/Source/OnlineSubsystemEOS/Private/OnlineSubsystemEOS.h
--- a/Engine/Plugins/Online/OnlineSubsystemEOS/Source/OnlineSubsystemEOS/Private/OnlineSubsystemEOS.h
+++ b/Engine/Plugins/Online/OnlineSubsystemEOS/Source/OnlineSubsystemEOS/Private/OnlineSubsystemEOS.h
@@ -53,6 +53,16 @@ typedef TSharedPtr<class FOnlinePlayerSanctionEOS, ESPMode::ThreadSafe> FOnlineP
 class FOnlinePlayerReportEOS;
 typedef TSharedPtr<class FOnlinePlayerReportEOS, ESPMode::ThreadSafe> FOnlinePlayerReportEOSPtr;
 
+// FOnlineEntitlementsEOS, FOnlinePurchaseEOS and FOnlineStoreEOSV2 are only set if bUseNewEcomFlow is true
+class FOnlineEntitlementsEOS;
+typedef TSharedPtr<class FOnlineEntitlementsEOS, ESPMode::ThreadSafe> FOnlineEntitlementsEOSPtr;
+
+class FOnlinePurchaseEOS;
+typedef TSharedPtr<class FOnlinePurchaseEOS, ESPMode::ThreadSafe> FOnlinePurchaseEOSPtr;
+
+class FOnlineStoreEOSV2;
+typedef TSharedPtr<class FOnlineStoreEOSV2, ESPMode::ThreadSafe> FOnlineStoreEOSV2Ptr;
+
 typedef TSharedPtr<FPlatformEOSHelpers, ESPMode::ThreadSafe> FPlatformEOSHelpersPtr;
 
 /**
@@ -162,6 +172,12 @@ public:
 	FOnlinePlayerSanctionEOSPtr PlayerSanctionEOSPtr;
 	/** Player Report interface pointer */
 	FOnlinePlayerReportEOSPtr PlayerReportInterfacePtr;
+	/** Entitlements interface pointer - only valid if bUseNewEcomFlow is true */
+	FOnlineEntitlementsEOSPtr EntitlementsInterfacePtr;
+	/** Purchase interface pointer - only valid if bUseNewEcomFlow is true */
+	FOnlinePurchaseEOSPtr PurchaseInterfacePtr;
+	/** Store interface pointer used to retrieve EGS offers - only valid if bUseNewEcomFlow is true */
+	FOnlineStoreEOSV2Ptr StoreInterfaceV2Ptr;
 
 	TSharedPtr<FSocketSubsystemEOS, ESPMode::ThreadSafe> SocketSubsystem;
 
@@ -205,6 +221,7 @@ public:
 	virtual IOnlineTitleFilePtr GetTitleFileInterface() const override { return nullptr; }
 	virtual IOnlineStoreV2Ptr GetStoreV2Interface() const override { return nullptr; }
 	virtual IOnlinePurchasePtr GetPurchaseInterface() const override { return nullptr; }
+	virtual IOnlineEntitlementsPtr GetEntitlementsInterface() const override { return nullptr; }
 	virtual IOnlineEventsPtr GetEventsInterface() const override { return nullptr; }
 	virtual IOnlineAchievementsPtr GetAchievementsInterface() const override { return nullptr; }
 	virtual IOnlineSharingPtr GetSharingInterface() const override { return nullptr; }
diff --git a/Engine/Plugins/Online/OnlineSubsystemEOS/Source/OnlineSubsystemEOS/Public/EOSSettings.h b/Engine/Plugins/Online/OnlineSubsystemEOS/Source/OnlineSubsystemEOS/Public/EOSSettings.h
--- a/Engine/Plugins/Online/OnlineSubsystemEOS/Source/OnlineSubsystemEOS/Public/EOSSettings.h
+++ b/Engine/Plugins/Online/OnlineSubsystemEOS/Source/OnlineSubsystemEOS/Public/EOSSettings.h
@@ -92,6 +92,7 @@ PRAGMA_ENABLE_DEPRECATION_WARNINGS
 	bool bUseEOSConnect;
 	bool bUseEOSRTC;
 	bool bUseNamedPlatformConfig;
+	bool bUseNewEcomFlow;
 	UE_DEPRECATED(5.7, "bUseNewLoginFlow is deprecated, the legacy login flow has been removed from UserManagerEOS.")
 	bool bUseNewLoginFlow;
 	TArray<FEOSArtifactSettings> Artifacts;
@@ -182,6 +183,10 @@ public:
 	UPROPERTY(Config, EditAnywhere, BlueprintReadOnly, Category="EOS Settings")
 	bool bUseNamedPlatformConfig = false;
 
+	/** New Epic Games Store Ecom integration - Use IOnlineEntitlements, IOnlinePurchase and IOnlineStoreV2 for a complete feature set of the Ecom interface in the EOS SDK */
+	UPROPERTY(Config, EditAnywhere, BlueprintReadOnly, Category = "EOS Ecom Settings", DisplayName = "Use new Ecom flow")
+	bool bUseNewEcomFlow = false;
+
 	/** Set to true to use new EOS login flow */
 	UE_DEPRECATED(5.7, "bUseNewLoginFlow is deprecated, the legacy login flow has been removed from UserManagerEOS.")
 	UPROPERTY(Config, EditAnywhere, BlueprintReadOnly, Category="EOS Login Settings", DisplayName="Use new Login flow, which doesn't rely on EOSPlus")
