Added check to prevent crash in tests when GEngine is null

#Synced-CL 48832497

diff --git a/Engine/Plugins/Online/OnlineSubsystemEOS/Source/OnlineSubsystemEOS/Private/OnlineSessionEOS.cpp b/Engine/Plugins/Online/OnlineSubsystemEOS/Source/OnlineSubsystemEOS/Private/OnlineSessionEOS.cpp
--- a/Engine/Plugins/Online/OnlineSubsystemEOS/Source/OnlineSubsystemEOS/Private/OnlineSessionEOS.cpp
+++ b/Engine/Plugins/Online/OnlineSubsystemEOS/Source/OnlineSubsystemEOS/Private/OnlineSessionEOS.cpp
@@ -1559,12 +1559,19 @@ TSharedPtr<class FInternetAddr> GetHostAddress(const FString& EosHostAddress, co
 
 bool IsNetDriverEOS(const FName NetDriverName)
 {
-	const FNetDriverDefinition* Found = GEngine->NetDriverDefinitions.FindByPredicate([NetDriverName](const FNetDriverDefinition& Elem)
+	if (GEngine)
+	{
+		const FNetDriverDefinition* Found = GEngine->NetDriverDefinitions.FindByPredicate([NetDriverName](const FNetDriverDefinition& Elem)
 		{
 			return Elem.DefName == NetDriverName;
 		});
 
-	return Found && Found->DriverClassName.ToString().Contains(TEXT("NetDriverEOS"));
+		return Found && Found->DriverClassName.ToString().Contains(TEXT("NetDriverEOS"));
+	}
+
+	UE_LOG_ONLINE_SESSION(Verbose, TEXT("GEngine not valid. Unable to check if NetDriver [%s] was EOS."), *NetDriverName.ToString());
+
+	return false;
 }
 
 uint32 FOnlineSessionEOS::CreateEOSSession(int32 HostingPlayerNum, FNamedOnlineSession* Session)
