Wrong EOS Notification ID object used for JoinSession callback and missing release, and delete in destructor.

#Synced-CL 49927350

diff --git a/Engine/Plugins/Online/OnlineSubsystemEOS/Source/OnlineSubsystemEOS/Private/OnlineSessionEOS.cpp b/Engine/Plugins/Online/OnlineSubsystemEOS/Source/OnlineSubsystemEOS/Private/OnlineSessionEOS.cpp
--- a/Engine/Plugins/Online/OnlineSubsystemEOS/Source/OnlineSubsystemEOS/Private/OnlineSessionEOS.cpp
+++ b/Engine/Plugins/Online/OnlineSubsystemEOS/Source/OnlineSubsystemEOS/Private/OnlineSessionEOS.cpp
@@ -435,14 +435,19 @@ typedef TEOSGlobalCallback<EOS_Lobby_OnLeaveLobbyRequestedCallback, EOS_Lobby_Le
 
 FOnlineSessionEOS::~FOnlineSessionEOS()
 {
+	EOS_Sessions_RemoveNotifySessionInviteReceived(EOSSubsystem->SessionsHandle, SessionInviteReceivedId);
 	EOS_Sessions_RemoveNotifySessionInviteAccepted(EOSSubsystem->SessionsHandle, SessionInviteAcceptedId);
+	EOS_Sessions_RemoveNotifyJoinSessionAccepted(EOSSubsystem->SessionsHandle, JoinSessionAcceptedId);
 	EOS_Sessions_RemoveNotifyLeaveSessionRequested(EOSSubsystem->SessionsHandle, LeaveSessionRequestedId);
+	delete SessionInviteReceivedCallback;
 	delete SessionInviteAcceptedCallback;
+	delete JoinSessionAcceptedCallback;
 	delete LeaveSessionRequestedCallback;
 
 	EOS_Lobby_RemoveNotifyLobbyUpdateReceived(LobbyHandle, LobbyUpdateReceivedId);
 	EOS_Lobby_RemoveNotifyLobbyMemberUpdateReceived(LobbyHandle, LobbyMemberUpdateReceivedId);
 	EOS_Lobby_RemoveNotifyLobbyMemberStatusReceived(LobbyHandle, LobbyMemberStatusReceivedId);
+	EOS_Lobby_RemoveNotifyLobbyInviteReceived(LobbyHandle, LobbyInviteReceivedId);
 	EOS_Lobby_RemoveNotifyLobbyInviteAccepted(LobbyHandle, LobbyInviteAcceptedId);
 	EOS_Lobby_RemoveNotifyJoinLobbyAccepted(LobbyHandle, JoinLobbyAcceptedId);
 	EOS_Lobby_RemoveNotifyLeaveLobbyRequested(LobbyHandle, LeaveLobbyRequestedId);
@@ -450,6 +455,7 @@ FOnlineSessionEOS::~FOnlineSessionEOS()
 	delete LobbyUpdateReceivedCallback;
 	delete LobbyMemberUpdateReceivedCallback;
 	delete LobbyMemberStatusReceivedCallback;
+	delete LobbyInviteReceivedCallback;
 	delete LobbyInviteAcceptedCallback;
 	delete JoinLobbyAcceptedCallback;
 	delete LeaveLobbyRequestedCallback;
@@ -626,7 +632,7 @@ void FOnlineSessionEOS::RegisterSessionNotifications()
 	EOS_Sessions_AddNotifyJoinSessionAcceptedOptions AddNotifyJoinSessionAcceptedOptions = { };
 	AddNotifyJoinSessionAcceptedOptions.ApiVersion = 1;
 	UE_EOS_CHECK_API_MISMATCH(EOS_SESSIONS_ADDNOTIFYJOINSESSIONACCEPTED_API_LATEST, 1);
-	SessionInviteAcceptedId = EOS_Sessions_AddNotifyJoinSessionAccepted(EOSSubsystem->SessionsHandle, &AddNotifyJoinSessionAcceptedOptions, JoinSessionAcceptedCallbackObj, JoinSessionAcceptedCallbackObj->GetCallbackPtr());
+	JoinSessionAcceptedId = EOS_Sessions_AddNotifyJoinSessionAccepted(EOSSubsystem->SessionsHandle, &AddNotifyJoinSessionAcceptedOptions, JoinSessionAcceptedCallbackObj, JoinSessionAcceptedCallbackObj->GetCallbackPtr());
 
 	// Requested session leave notifications
 	EOS_Sessions_AddNotifyLeaveSessionRequestedOptions AddNotifyLeaveSessionRequestedOptions = { 0 };
