Log a warning when a stat name passed to FOnlineStatsEOS::WriteStats is not uppercase

#Synced-CL 50659603

diff --git a/Engine/Plugins/Online/OnlineSubsystemEOS/Source/OnlineSubsystemEOS/Private/OnlineStatsEOS.cpp b/Engine/Plugins/Online/OnlineSubsystemEOS/Source/OnlineSubsystemEOS/Private/OnlineStatsEOS.cpp
--- a/Engine/Plugins/Online/OnlineSubsystemEOS/Source/OnlineSubsystemEOS/Private/OnlineStatsEOS.cpp
+++ b/Engine/Plugins/Online/OnlineSubsystemEOS/Source/OnlineSubsystemEOS/Private/OnlineStatsEOS.cpp
@@ -283,21 +283,27 @@ typedef TEOSCallback<EOS_Stats_OnIngestStatCompleteCallback, EOS_Stats_IngestSta
 void FOnlineStatsEOS::WriteStats(EOS_ProductUserId LocalUserId, EOS_ProductUserId UserId, const FOnlineStatsUserUpdatedStats& PlayerStats)
 {
 	TArray<EOS_Stats_IngestData> EOSData;
-	TArray<FStatNameBuffer> EOSStatNames;
+	TArray<FTCHARToUTF8> EOSStatNamesUtf8;
 	// Preallocate all of the memory
 	EOSData.AddZeroed(PlayerStats.Stats.Num());
-	EOSStatNames.AddZeroed(PlayerStats.Stats.Num());
+	EOSStatNamesUtf8.Reserve(PlayerStats.Stats.Num());
 	uint32 Index = 0;
 	// Convert the stats to the EOS format
 	for (const TPair<FString, FOnlineStatUpdate>& Stat : PlayerStats.Stats)
 	{
+		FString StatKeyUppercase = Stat.Key.ToUpper();
+		EOSStatNamesUtf8.Emplace(*StatKeyUppercase);
+
+		if (Stat.Key != StatKeyUppercase)
+		{
+			UE_LOG_ONLINE_STATS(Warning, TEXT("WriteStats() stat name (%s) is not upper case. Ingest will fail if the stat name is not upper case in the Developer Portal."), *Stat.Key);
+		}
+
 		EOS_Stats_IngestData& EOSStat = EOSData[Index];
 		EOSStat.ApiVersion = 1;
 		UE_EOS_CHECK_API_MISMATCH(EOS_STATS_INGESTDATA_API_LATEST, 1);
-
 		EOSStat.IngestAmount = GetVariantValue(Stat.Value.GetValue());
-		FCStringAnsi::Strncpy(EOSStatNames[Index].StatName, TCHAR_TO_UTF8(*Stat.Key.ToUpper()), EOS_OSS_STRING_BUFFER_LENGTH);
-		EOSStat.StatName = EOSStatNames[Index].StatName;
+		EOSStat.StatName = EOSStatNamesUtf8.Last().Get();
 
 		Index++;
 	}
