Fix crash when Connect AddNotifyAuthExpiration callback fires if the user logs out and logs back in
Cleanup notification and callback in UserManagerEOS
Other minor fixes

#Synced-CL 50659603

diff --git a/Engine/Plugins/Online/OnlineSubsystemEOS/Source/OnlineSubsystemEOS/Private/UserManagerEOS.cpp b/Engine/Plugins/Online/OnlineSubsystemEOS/Source/OnlineSubsystemEOS/Private/UserManagerEOS.cpp
--- a/Engine/Plugins/Online/OnlineSubsystemEOS/Source/OnlineSubsystemEOS/Private/UserManagerEOS.cpp
+++ b/Engine/Plugins/Online/OnlineSubsystemEOS/Source/OnlineSubsystemEOS/Private/UserManagerEOS.cpp
@@ -266,8 +266,12 @@ FUserManagerEOS::FUserManagerEOS(FOnlineSubsystemEOS* InSubsystem)
 	: TSharedFromThis<FUserManagerEOS, ESPMode::ThreadSafe>()
 	, EOSSubsystem(InSubsystem)
 	, DefaultLocalUser(INVALID_LOCAL_USER)
-	, LoginNotificationId(0)
-	, LoginNotificationCallback(nullptr)
+	, AuthLoginStatusChangedNotificationId(0)
+	, AuthLoginStatusChangedNotificationCallback(nullptr)
+	, ConnectLoginStatusChangedNotificationId(0)
+	, ConnectLoginStatusChangedNotificationCallback(nullptr)
+	, ConnectLoginExpiredNotificationId(0)
+	, ConnectLoginExpiredNotificationCallback(nullptr)
 	, FriendsNotificationId(0)
 	, FriendsNotificationCallback(nullptr)
 	, PresenceNotificationId(0)
@@ -277,40 +281,155 @@ FUserManagerEOS::FUserManagerEOS(FOnlineSubsystemEOS* InSubsystem)
 {
 }
 
+typedef TEOSGlobalCallback<EOS_Connect_OnAuthExpirationCallback, EOS_Connect_AuthExpirationCallbackInfo, FUserManagerEOS> FRefreshAuthCallback;
+typedef TEOSGlobalCallback<EOS_Auth_OnLoginStatusChangedCallback, EOS_Auth_LoginStatusChangedCallbackInfo, FUserManagerEOS> FAuthLoginStatusChangedCallback;
+typedef TEOSGlobalCallback<EOS_Connect_OnLoginStatusChangedCallback, EOS_Connect_LoginStatusChangedCallbackInfo, FUserManagerEOS> FConnectLoginStatusChangedCallback;
+typedef TEOSGlobalCallback<EOS_Friends_OnFriendsUpdateCallback, EOS_Friends_OnFriendsUpdateInfo, FUserManagerEOS> FFriendsStatusUpdateCallback;
+typedef TEOSGlobalCallback<EOS_Connect_OnAuthExpirationCallback, EOS_Connect_AuthExpirationCallbackInfo, FUserManagerEOS> FRefreshConnectCallback;
+typedef TEOSGlobalCallback<EOS_Presence_OnPresenceChangedCallback, EOS_Presence_PresenceChangedCallbackInfo, FUserManagerEOS> FPresenceChangedCallback;
+
 void FUserManagerEOS::Init()
 {
-	// This delegate would cause a crash when running a dedicated server
+	// Don't bind player notifications on a dedicated server
 	if (!IsRunningDedicatedServer())
 	{
-		// Adding subscription to external ui display change event
-		EOS_UI_AddNotifyDisplaySettingsUpdatedOptions Options = {};
-		Options.ApiVersion = 1;
-		UE_EOS_CHECK_API_MISMATCH(EOS_UI_ADDNOTIFYDISPLAYSETTINGSUPDATED_API_LATEST, 1);
+		FEOSSettings Settings = UEOSSettings::GetSettings();
 
-		FOnDisplaySettingsUpdatedCallback* CallbackObj = new FOnDisplaySettingsUpdatedCallback(AsWeak());
-		DisplaySettingsUpdatedCallback = CallbackObj;
-		CallbackObj->CallbackLambda = [this](const EOS_UI_OnDisplaySettingsUpdatedCallbackInfo* Data)
+		if (Settings.bUseEAS)
 		{
-			TriggerOnExternalUIChangeDelegates((bool)Data->bIsVisible);
-		};
+			FAuthLoginStatusChangedCallback* AuthLoginStatusCallbackObj = new FAuthLoginStatusChangedCallback(AsWeak());
+			AuthLoginStatusChangedNotificationCallback = AuthLoginStatusCallbackObj;
+			AuthLoginStatusCallbackObj->CallbackLambda = [this](const EOS_Auth_LoginStatusChangedCallbackInfo* Data)
+				{
+					AuthLoginStatusChanged(Data);
+				};
+
+			EOS_Auth_AddNotifyLoginStatusChangedOptions AuthLoginStatusOptions = { };
+			AuthLoginStatusOptions.ApiVersion = 1;
+			UE_EOS_CHECK_API_MISMATCH(EOS_AUTH_ADDNOTIFYLOGINSTATUSCHANGED_API_LATEST, 1);
+			AuthLoginStatusChangedNotificationId = EOS_Auth_AddNotifyLoginStatusChanged(EOSSubsystem->AuthHandle, &AuthLoginStatusOptions, AuthLoginStatusCallbackObj, AuthLoginStatusCallbackObj->GetCallbackPtr());
+
+			FFriendsStatusUpdateCallback* FriendsStatusCallbackObj = new FFriendsStatusUpdateCallback(AsWeak());
+			FriendsNotificationCallback = FriendsStatusCallbackObj;
+			FriendsStatusCallbackObj->CallbackLambda = [this](const EOS_Friends_OnFriendsUpdateInfo* Data)
+				{
+					FriendStatusChanged(Data);
+				};
+
+			EOS_Friends_AddNotifyFriendsUpdateOptions FriendsStatusOptions = { };
+			FriendsStatusOptions.ApiVersion = 1;
+			UE_EOS_CHECK_API_MISMATCH(EOS_FRIENDS_ADDNOTIFYFRIENDSUPDATE_API_LATEST, 1);
+			FriendsNotificationId = EOS_Friends_AddNotifyFriendsUpdate(EOSSubsystem->FriendsHandle, &FriendsStatusOptions, FriendsStatusCallbackObj, FriendsStatusCallbackObj->GetCallbackPtr());
+
+			FPresenceChangedCallback* PresenceChangedCallbackObj = new FPresenceChangedCallback(AsWeak());
+			PresenceNotificationCallback = PresenceChangedCallbackObj;
+			PresenceChangedCallbackObj->CallbackLambda = [this](const EOS_Presence_PresenceChangedCallbackInfo* Data)
+				{
+					PresenceChanged(Data);
+				};
+
+			EOS_Presence_AddNotifyOnPresenceChangedOptions PresenceChangedOptions = { };
+			PresenceChangedOptions.ApiVersion = 1;
+			UE_EOS_CHECK_API_MISMATCH(EOS_PRESENCE_ADDNOTIFYONPRESENCECHANGED_API_LATEST, 1);
+			PresenceNotificationId = EOS_Presence_AddNotifyOnPresenceChanged(EOSSubsystem->PresenceHandle, &PresenceChangedOptions, PresenceChangedCallbackObj, PresenceChangedCallbackObj->GetCallbackPtr());
+
+			FOnDisplaySettingsUpdatedCallback* DisplaySettingsUpdatedCallbackObj = new FOnDisplaySettingsUpdatedCallback(AsWeak());
+			DisplaySettingsUpdatedCallback = DisplaySettingsUpdatedCallbackObj;
+			DisplaySettingsUpdatedCallbackObj->CallbackLambda = [this](const EOS_UI_OnDisplaySettingsUpdatedCallbackInfo* Data)
+				{
+					TriggerOnExternalUIChangeDelegates((bool)Data->bIsVisible);
+				};
+
+			EOS_UI_AddNotifyDisplaySettingsUpdatedOptions DisplaySettingsUpdatedOptions = {};
+			DisplaySettingsUpdatedOptions.ApiVersion = 1;
+			UE_EOS_CHECK_API_MISMATCH(EOS_UI_ADDNOTIFYDISPLAYSETTINGSUPDATED_API_LATEST, 1);
+			DisplaySettingsUpdatedId = EOS_UI_AddNotifyDisplaySettingsUpdated(EOSSubsystem->UIHandle, &DisplaySettingsUpdatedOptions, DisplaySettingsUpdatedCallbackObj, DisplaySettingsUpdatedCallbackObj->GetCallbackPtr());
+		}
+		
+		FConnectLoginStatusChangedCallback* ConnectLoginStatusCallbackObj = new FConnectLoginStatusChangedCallback(AsWeak());
+		ConnectLoginStatusChangedNotificationCallback = ConnectLoginStatusCallbackObj;
+		ConnectLoginStatusCallbackObj->CallbackLambda = [this](const EOS_Connect_LoginStatusChangedCallbackInfo* Data)
+			{
+				ConnectLoginStatusChanged(Data);
+			};
+
+		EOS_Connect_AddNotifyLoginStatusChangedOptions ConnectLoginStatusOptions = { };
+		ConnectLoginStatusOptions.ApiVersion = 1;
+		UE_EOS_CHECK_API_MISMATCH(EOS_CONNECT_ADDNOTIFYLOGINSTATUSCHANGED_API_LATEST, 1);
+		ConnectLoginStatusChangedNotificationId = EOS_Connect_AddNotifyLoginStatusChanged(EOSSubsystem->ConnectHandle, &ConnectLoginStatusOptions, ConnectLoginStatusCallbackObj, ConnectLoginStatusCallbackObj->GetCallbackPtr());
+		
+		FRefreshAuthCallback* CallbackObj = new FRefreshConnectCallback(AsWeak());
+		ConnectLoginExpiredNotificationCallback = CallbackObj;
+		CallbackObj->CallbackLambda = [this](const EOS_Connect_AuthExpirationCallbackInfo* Data)
+			{
+				AutoRefreshConnectLogin(Data);
+			};
 
-		DisplaySettingsUpdatedId = EOS_UI_AddNotifyDisplaySettingsUpdated(EOSSubsystem->UIHandle, &Options, CallbackObj, CallbackObj->GetCallbackPtr());
+		EOS_Connect_AddNotifyAuthExpirationOptions Options = { };
+		Options.ApiVersion = 1;
+		UE_EOS_CHECK_API_MISMATCH(EOS_CONNECT_ADDNOTIFYAUTHEXPIRATION_API_LATEST, 1);
+		ConnectLoginExpiredNotificationId = EOS_Connect_AddNotifyAuthExpiration(EOSSubsystem->ConnectHandle, &Options, CallbackObj, CallbackObj->GetCallbackPtr());	
 	}
 }
 
 void FUserManagerEOS::Shutdown()
 {
-	// This delegate would cause a crash when running a dedicated server
-	if (DisplaySettingsUpdatedId != EOS_INVALID_NOTIFICATIONID)
+	FEOSSettings Settings = UEOSSettings::GetSettings();
+
+	if (Settings.bUseEAS)
 	{
-		// Removing subscription to external ui display change event
-		EOS_UI_RemoveNotifyDisplaySettingsUpdated(EOSSubsystem->UIHandle, DisplaySettingsUpdatedId);
+		if (AuthLoginStatusChangedNotificationId > 0)
+		{
+			EOS_Auth_RemoveNotifyLoginStatusChanged(EOSSubsystem->AuthHandle, AuthLoginStatusChangedNotificationId);
+		}
+		delete AuthLoginStatusChangedNotificationCallback;
+		AuthLoginStatusChangedNotificationCallback = nullptr;
+		AuthLoginStatusChangedNotificationId = 0;
 
-		if (DisplaySettingsUpdatedCallback)
+		if (DisplaySettingsUpdatedId > 0)
 		{
-			delete DisplaySettingsUpdatedCallback;
+			// Removing subscription to external ui display change event
+			EOS_UI_RemoveNotifyDisplaySettingsUpdated(EOSSubsystem->UIHandle, DisplaySettingsUpdatedId);
 		}
+		delete DisplaySettingsUpdatedCallback;
+		DisplaySettingsUpdatedCallback = nullptr;
+		DisplaySettingsUpdatedId = 0;
+
+		if (FriendsNotificationId > 0)
+		{
+			EOS_Friends_RemoveNotifyFriendsUpdate(EOSSubsystem->FriendsHandle, FriendsNotificationId);
+		}
+		delete FriendsNotificationCallback;
+		FriendsNotificationCallback = nullptr;
+		FriendsNotificationId = 0;
+
+		if (PresenceNotificationId > 0)
+		{
+			EOS_Presence_RemoveNotifyOnPresenceChanged(EOSSubsystem->PresenceHandle, PresenceNotificationId);
+		}
+
+		delete PresenceNotificationCallback;
+		PresenceNotificationCallback = nullptr;
+		PresenceNotificationId = 0;
 	}
+
+	if (ConnectLoginStatusChangedNotificationId > 0)
+	{
+		EOS_Connect_RemoveNotifyLoginStatusChanged(EOSSubsystem->ConnectHandle, ConnectLoginStatusChangedNotificationId);
+	}
+
+	delete ConnectLoginStatusChangedNotificationCallback;
+	ConnectLoginStatusChangedNotificationCallback = nullptr;
+	ConnectLoginStatusChangedNotificationId = 0;
+
+	if (ConnectLoginExpiredNotificationId > 0)
+	{
+		EOS_Connect_RemoveNotifyAuthExpiration(EOSSubsystem->ConnectHandle, ConnectLoginExpiredNotificationId);
+	}
+
+	delete ConnectLoginExpiredNotificationCallback;
+	ConnectLoginExpiredNotificationCallback = nullptr;
+	ConnectLoginExpiredNotificationId = 0;
 }
 
 FUserManagerEOS::~FUserManagerEOS()
@@ -318,54 +437,91 @@ FUserManagerEOS::~FUserManagerEOS()
 	Shutdown();
 }
 
-void FUserManagerEOS::LoginStatusChanged(const EOS_Auth_LoginStatusChangedCallbackInfo* Data)
+typedef TEOSCallback<EOS_Auth_OnLogoutCallback, EOS_Auth_LogoutCallbackInfo, FUserManagerEOS> FAuthLogoutCallback;
+typedef TEOSCallback<EOS_Connect_OnLogoutCallback, EOS_Connect_LogoutCallbackInfo, FUserManagerEOS> FConnectLogoutCallback;
+void FUserManagerEOS::AuthLoginStatusChanged(const EOS_Auth_LoginStatusChangedCallbackInfo* Data)
 {
-	if (Data->CurrentStatus == EOS_ELoginStatus::EOS_LS_NotLoggedIn)
+	const int32 LocalUserNum = GetLocalUserNumFromEpicAccountId(Data->LocalUserId);
+	FString EpicAccountIdStr = LexToString(Data->LocalUserId);
+	FString RedactedEpicAccountIdStr = EpicAccountIdStr.Left(3) + TEXT("...") + EpicAccountIdStr.Right(3);
+	
+	// LocalUserNum may be -1 if this callback was triggered before the login callback
+	UE_LOG_ONLINE(Verbose, TEXT("[%hs] LocalUserNum=[%d] EpicAccountId=[%s] PreviousStatus=[%s] CurrentStatus=[%s]"), __FUNCTION__, LocalUserNum, *RedactedEpicAccountIdStr, LexToString(Data->PrevStatus), LexToString(Data->CurrentStatus));
+
+	if (IsLocalUserValid(LocalUserNum) && Data->CurrentStatus == EOS_ELoginStatus::EOS_LS_NotLoggedIn && Data->PrevStatus == EOS_ELoginStatus::EOS_LS_LoggedIn)
 	{
-		if (IsLocalUser(Data->LocalUserId))
+		FLocalUserEOS& LocalUser = GetLocalUserChecked(LocalUserNum);
+		const FUniqueNetIdEOSPtr& UserNetId = LocalUser.UniqueNetId;
+
+		// Logged out of EAS and logged in to Game Services is an unexpected state, force a logout of Game Services only if a logout isn't currently in flight
+		// Notification and local user cleanup happen in ConnectLoginStatusChanged
+		EOS_ELoginStatus ConnectLoginStatus = EOS_Connect_GetLoginStatus(EOSSubsystem->ConnectHandle, UserNetId->GetProductUserId());
+		if (ConnectLoginStatus == EOS_ELoginStatus::EOS_LS_LoggedIn && !LocalUser.bIsLoggingOut)
 		{
-			const int32 LocalUserNum = GetLocalUserNumFromEpicAccountId(Data->LocalUserId);
-			FLocalUserEOS& LocalUser = GetLocalUserChecked(LocalUserNum);
-			const FUniqueNetIdEOSPtr& UserNetId = LocalUser.UniqueNetId;
-			TriggerOnLoginStatusChangedDelegates(LocalUserNum, ELoginStatus::LoggedIn, ELoginStatus::NotLoggedIn, *UserNetId);
+			UE_LOG_ONLINE(Log, TEXT("[%hs] LocalUserNum=[%d] Logged out of EAS, but still logged in Game Service - Logging out of Game Service"), __FUNCTION__, LocalUserNum);
 
-			// Remove the per user connect login notification
-			if (LocalUser.ConnectLoginNotification.IsValid())
-			{
-				EOS_Connect_RemoveNotifyAuthExpiration(EOSSubsystem->ConnectHandle, LocalUser.ConnectLoginNotification->NotificationId);
-				LocalUser.ConnectLoginNotification.Reset();
-			}
+			FConnectLogoutCallback* CallbackObj = new FConnectLogoutCallback(AsWeak());
+			CallbackObj->CallbackLambda = [LocalUserNum, this](const EOS_Connect_LogoutCallbackInfo* Data)
+				{
+					const bool bWasConnectLogoutSuccessful = Data->ResultCode == EOS_EResult::EOS_Success;
 
-			// Need to remove the local user
-			RemoveLocalUser(LocalUserNum);
+					UE_CLOG_ONLINE(!bWasConnectLogoutSuccessful, Warning, TEXT("EOS_Connect_Logout Result=[%s]"), *LexToString(Data->ResultCode));
+				};
 
-			// Clean up user based notifications if we have no logged in users
-			if (LocalUsers.Num() == 0)
+			EOS_Connect_LogoutOptions LogoutOptions = { };
+			LogoutOptions.ApiVersion = 1;
+			UE_EOS_CHECK_API_MISMATCH(EOS_CONNECT_LOGOUT_API_LATEST, 1);
+			LogoutOptions.LocalUserId = UserNetId->GetProductUserId();
+			EOS_Connect_Logout(EOSSubsystem->ConnectHandle, &LogoutOptions, CallbackObj, CallbackObj->GetCallbackPtr());
+		}
+	}
+}
+
+void FUserManagerEOS::ConnectLoginStatusChanged(const EOS_Connect_LoginStatusChangedCallbackInfo* Data)
+{
+	const int32 LocalUserNum = GetLocalUserNumFromProductUserId(Data->LocalUserId);
+
+	FString ProductUserIdStr = LexToString(Data->LocalUserId);
+	FString RedactedProductUserIdStr = ProductUserIdStr.Left(3) + TEXT("...") + ProductUserIdStr.Right(3);
+
+	// LocalUserNum may be -1 if this callback was triggered before the login callback
+	UE_LOG_ONLINE(Verbose, TEXT("[%hs] LocalUserNum=[%d] ProductUserId=[%s] PreviousStatus=[%s] CurrentStatus=[%s]"), __FUNCTION__, LocalUserNum, *RedactedProductUserIdStr, LexToString(Data->PreviousStatus), LexToString(Data->CurrentStatus));
+
+	if (IsLocalUserValid(LocalUserNum) && Data->CurrentStatus == EOS_ELoginStatus::EOS_LS_NotLoggedIn && Data->PreviousStatus == EOS_ELoginStatus::EOS_LS_LoggedIn)
+	{
+		FLocalUserEOS& LocalUser = GetLocalUserChecked(LocalUserNum);
+		const FUniqueNetIdEOSPtr& UserNetId = LocalUser.UniqueNetId;
+
+		FEOSSettings Settings = UEOSSettings::GetSettings();
+
+		if (Settings.bUseEAS)
+		{
+			// Logged into EAS and logged out of Game Services is an unexpected state, force a logout of EAS
+			EOS_ELoginStatus AuthLoginStatus = EOS_Auth_GetLoginStatus(EOSSubsystem->AuthHandle, UserNetId->GetEpicAccountId());
+			if (AuthLoginStatus == EOS_ELoginStatus::EOS_LS_LoggedIn && !LocalUser.bIsLoggingOut)
 			{
-				if (LoginNotificationId > 0)
-				{
-					// Remove the callback
-					EOS_Auth_RemoveNotifyLoginStatusChanged(EOSSubsystem->AuthHandle, LoginNotificationId);
-					delete LoginNotificationCallback;
-					LoginNotificationCallback = nullptr;
-					LoginNotificationId = 0;
-				}
-				if (FriendsNotificationId > 0)
-				{
-					EOS_Friends_RemoveNotifyFriendsUpdate(EOSSubsystem->FriendsHandle, FriendsNotificationId);
-					delete FriendsNotificationCallback;
-					FriendsNotificationCallback = nullptr;
-					FriendsNotificationId = 0;
-				}
-				if (PresenceNotificationId > 0)
-				{
-					EOS_Presence_RemoveNotifyOnPresenceChanged(EOSSubsystem->PresenceHandle, PresenceNotificationId);
-					delete PresenceNotificationCallback;
-					PresenceNotificationCallback = nullptr;
-					PresenceNotificationId = 0;
-				}
+				UE_LOG_ONLINE(Log, TEXT("[%hs] LocalUserNum: %d Logged out of Game Services, but still logged in EAS - Logging out of EAS"), __FUNCTION__, LocalUserNum);
+
+				FAuthLogoutCallback* CallbackObj = new FAuthLogoutCallback(AsWeak());
+				CallbackObj->CallbackLambda = [LocalUserNum, this](const EOS_Auth_LogoutCallbackInfo* Data)
+					{
+						const bool bWasAuthLogoutSuccessful = Data->ResultCode == EOS_EResult::EOS_Success;
+
+						UE_CLOG_ONLINE(!bWasAuthLogoutSuccessful, Warning, TEXT("EOS_Auth_Logout Result=[%s]"), *LexToString(Data->ResultCode));
+					};
+
+				EOS_Auth_LogoutOptions LogoutOptions = { };
+				LogoutOptions.ApiVersion = 1;
+				UE_EOS_CHECK_API_MISMATCH(EOS_AUTH_LOGOUT_API_LATEST, 1);
+				LogoutOptions.LocalUserId = UserNetId->GetEpicAccountId();
+
+				EOS_Auth_Logout(EOSSubsystem->AuthHandle, &LogoutOptions, CallbackObj, CallbackObj->GetCallbackPtr());
 			}
 		}
+		TriggerOnLoginStatusChangedDelegates(LocalUserNum, ELoginStatus::LoggedIn, ELoginStatus::NotLoggedIn, *UserNetId);
+
+		// Need to remove the local user
+		RemoveLocalUser(LocalUserNum);
 	}
 }
 
@@ -1100,174 +1256,179 @@ bool FUserManagerEOS::ConnectLoginEAS(int32 LocalUserNum, EOS_EpicAccountId Acco
 	return true;
 }
 
-void FUserManagerEOS::AutoRefreshConnectLogin(int32 LocalUserNum)
+void FUserManagerEOS::AutoRefreshConnectLogin(const EOS_Connect_AuthExpirationCallbackInfo* Data)
 {
-	const FEOSSettings Settings = UEOSSettings::GetSettings();
-	if (Settings.bUseEAS)
+	const int32 LocalUserNum = GetLocalUserNumFromProductUserId(Data->LocalUserId);
+
+	if (IsLocalUserValid(LocalUserNum))
 	{
-		const EOS_EpicAccountId AccountId = GetLocalEpicAccountId(LocalUserNum);
-		if (!AccountId)
-		{
-			UE_LOG_ONLINE(Error, TEXT("Can't refresh ConnectLogin(%d) since (%d) is not logged in"), LocalUserNum, LocalUserNum);
-			return;
-		}
-		const FString AccessToken = GetAuthToken(LocalUserNum);
-		if (!AccessToken.IsEmpty())
+		const FEOSSettings Settings = UEOSSettings::GetSettings();
+		if (Settings.bUseEAS)
 		{
+			const EOS_EpicAccountId AccountId = GetLocalEpicAccountId(LocalUserNum);
+			if (!AccountId)
+			{
+				UE_LOG_ONLINE(Error, TEXT("Can't refresh ConnectLogin(%d) since (%d) is not logged in"), LocalUserNum, LocalUserNum);
+				return;
+			}
+			const FString AccessToken = GetAuthToken(LocalUserNum);
+			if (!AccessToken.IsEmpty())
+			{
 #if ADD_USER_LOGIN_INFO
-			GetPlatformAuthToken(LocalUserNum,
-				FOnGetLinkedAccountAuthTokenCompleteDelegate::CreateLambda([this, WeakThis = AsWeak(), AccessToken, AccountId](int32 LocalUserNum, bool bWasSuccessful, const FExternalAuthToken& AuthToken)
-				{
-					if (FUserManagerEOSPtr StrongThis = WeakThis.Pin())
-					{
-						if (!bWasSuccessful || !AuthToken.IsValid())
+				GetPlatformAuthToken(LocalUserNum,
+					FOnGetLinkedAccountAuthTokenCompleteDelegate::CreateLambda([this, WeakThis = AsWeak(), AccessToken, AccountId](int32 LocalUserNum, bool bWasSuccessful, const FExternalAuthToken& AuthToken)
 						{
-							const FString ErrorString = FString::Printf(TEXT("ConnectLoginEAS(%d) failed due to the platform OSS giving an empty auth token"), LocalUserNum);
-							UE_LOG_ONLINE(Warning, TEXT("%s"), *ErrorString);
-							RemoveLocalUser(LocalUserNum);
-							TriggerOnLoginCompleteDelegates(LocalUserNum, false, *FUniqueNetIdEOS::EmptyId(), ErrorString);
-							return;
-						}
+							if (FUserManagerEOSPtr StrongThis = WeakThis.Pin())
+							{
+								if (!bWasSuccessful || !AuthToken.IsValid())
+								{
+									const FString ErrorString = FString::Printf(TEXT("ConnectLoginEAS(%d) failed due to the platform OSS giving an empty auth token"), LocalUserNum);
+									UE_LOG_ONLINE(Warning, TEXT("%s"), *ErrorString);
+									RemoveLocalUser(LocalUserNum);
+									TriggerOnLoginCompleteDelegates(LocalUserNum, false, *FUniqueNetIdEOS::EmptyId(), ErrorString);
+									return;
+								}
 #endif
-						// We update the auth token cached in the user account, along with the user information
-						const FUserOnlineAccountEOSRef UserAccountRef = GetLocalUserChecked(LocalUserNum).UserOnlineAccount.ToSharedRef();
-						UpdateUserInfo(UserAccountRef, AccountId, AccountId);
-
-						EOS_Connect_Credentials Credentials = { };
-						Credentials.ApiVersion = 1;
-						UE_EOS_CHECK_API_MISMATCH(EOS_CONNECT_CREDENTIALS_API_LATEST, 1);
-						Credentials.Type = EOS_EExternalCredentialType::EOS_ECT_EPIC;
-						auto AccessTokenConverter = StringCast<UTF8CHAR>(*AccessToken);
-						Credentials.Token = (const char*)AccessTokenConverter.Get();
-
-						EOS_Connect_LoginOptions Options = { };
-						Options.ApiVersion = 2;
-						UE_EOS_CHECK_API_MISMATCH(EOS_CONNECT_LOGIN_API_LATEST, 2);
-						Options.Credentials = &Credentials;
+								// We update the auth token cached in the user account, along with the user information
+								const FUserOnlineAccountEOSRef UserAccountRef = GetLocalUserChecked(LocalUserNum).UserOnlineAccount.ToSharedRef();
+								UpdateUserInfo(UserAccountRef, AccountId, AccountId);
+
+								EOS_Connect_Credentials Credentials = { };
+								Credentials.ApiVersion = 1;
+								UE_EOS_CHECK_API_MISMATCH(EOS_CONNECT_CREDENTIALS_API_LATEST, 1);
+								Credentials.Type = EOS_EExternalCredentialType::EOS_ECT_EPIC;
+								auto AccessTokenConverter = StringCast<UTF8CHAR>(*AccessToken);
+								Credentials.Token = (const char*)AccessTokenConverter.Get();
+
+								EOS_Connect_LoginOptions Options = { };
+								Options.ApiVersion = 2;
+								UE_EOS_CHECK_API_MISMATCH(EOS_CONNECT_LOGIN_API_LATEST, 2);
+								Options.Credentials = &Credentials;
 
 #if ADD_USER_LOGIN_INFO
-						EOS_Connect_UserLoginInfo UserLoginInfo = {};
-						UserLoginInfo.ApiVersion = 2;
-						UE_EOS_CHECK_API_MISMATCH(EOS_CONNECT_USERLOGININFO_API_LATEST, 2);
-						auto AuthTokenConverter = StringCast<UTF8CHAR>(*AuthToken.TokenString);
-						UserLoginInfo.NsaIdToken = (const char*)AuthTokenConverter.Get();
-						// UserLoginInfo.DisplayName should not be set when the Credentials Type is EPIC
-
-						Options.UserLoginInfo = &UserLoginInfo;
+								EOS_Connect_UserLoginInfo UserLoginInfo = {};
+								UserLoginInfo.ApiVersion = 2;
+								UE_EOS_CHECK_API_MISMATCH(EOS_CONNECT_USERLOGININFO_API_LATEST, 2);
+								auto AuthTokenConverter = StringCast<UTF8CHAR>(*AuthToken.TokenString);
+								UserLoginInfo.NsaIdToken = (const char*)AuthTokenConverter.Get();
+								// UserLoginInfo.DisplayName should not be set when the Credentials Type is EPIC
+
+								Options.UserLoginInfo = &UserLoginInfo;
 #endif
 
-						FConnectLoginCallback* CallbackObj = new FConnectLoginCallback(AsWeak());
-						CallbackObj->CallbackLambda = [LocalUserNum, AccountId, this](const EOS_Connect_LoginCallbackInfo* Data)
-							{
-								if (Data->ResultCode != EOS_EResult::EOS_Success)
-								{
-									UE_LOG_ONLINE(Error, TEXT("Failed to refresh ConnectLogin(%d) failed with EOS result code (%s)"), LocalUserNum, *LexToString(Data->ResultCode));
-									Logout(LocalUserNum);
-								}
-							};
-						EOS_Connect_Login(EOSSubsystem->ConnectHandle, &Options, CallbackObj, CallbackObj->GetCallbackPtr());
+								FConnectLoginCallback* CallbackObj = new FConnectLoginCallback(AsWeak());
+								CallbackObj->CallbackLambda = [LocalUserNum, AccountId, this](const EOS_Connect_LoginCallbackInfo* Data)
+									{
+										if (Data->ResultCode != EOS_EResult::EOS_Success)
+										{
+											UE_LOG_ONLINE(Error, TEXT("Failed to refresh ConnectLogin(%d) failed with EOS result code (%s)"), LocalUserNum, *LexToString(Data->ResultCode));
+											Logout(LocalUserNum);
+										}
+									};
+								EOS_Connect_Login(EOSSubsystem->ConnectHandle, &Options, CallbackObj, CallbackObj->GetCallbackPtr());
 
 #if ADD_USER_LOGIN_INFO
-					}
-				}));
+							}
+						}));
 #endif
+			}
+			else
+			{
+				UE_LOG_ONLINE(Error, TEXT("[FUserManagerEOS::AutoRefreshConnectLogin] AccessToken for user [%d] is empty"), LocalUserNum);
+				Logout(LocalUserNum);
+			}
 		}
 		else
 		{
-			UE_LOG_ONLINE(Error, TEXT("[FUserManagerEOS::AutoRefreshConnectLogin] AccessToken for user [%d] is empty"), LocalUserNum);
-			Logout(LocalUserNum);
-		}
-	}
-	else
-	{
-		if (IsLocalUserValid(LocalUserNum))
-		{
-			const FLocalUserEOS& LocalUser = GetLocalUserChecked(LocalUserNum);
-
-			// If the user provided a Connect Auth Token, we'll try renewing the auth with it before we fire the corresponding delegate
-			if (LocalUser.UserProvidedConnectAuthToken.IsSet())
+			if (IsLocalUserValid(LocalUserNum))
 			{
-				IOnlineSubsystem* PlatformOSS = GetPlatformOSS();
-				EOS_EExternalCredentialType CredType = ToEOS_EExternalCredentialType(PlatformOSS ? PlatformOSS->GetSubsystemName() : EOSSubsystem->GetSubsystemName());
+				const FLocalUserEOS& LocalUser = GetLocalUserChecked(LocalUserNum);
+
+				// If the user provided a Connect Auth Token, we'll try renewing the auth with it before we fire the corresponding delegate
+				if (LocalUser.UserProvidedConnectAuthToken.IsSet())
+				{
+					IOnlineSubsystem* PlatformOSS = GetPlatformOSS();
+					EOS_EExternalCredentialType CredType = ToEOS_EExternalCredentialType(PlatformOSS ? PlatformOSS->GetSubsystemName() : EOSSubsystem->GetSubsystemName());
 
-				FConnectCredentials Credentials(CredType, *LocalUser.UserProvidedConnectAuthToken);
+					FConnectCredentials Credentials(CredType, *LocalUser.UserProvidedConnectAuthToken);
 
-				EOS_Connect_LoginOptions Options = { };
-				Options.ApiVersion = 2;
-				UE_EOS_CHECK_API_MISMATCH(EOS_CONNECT_LOGIN_API_LATEST, 2);
-				Options.Credentials = &Credentials;
+					EOS_Connect_LoginOptions Options = { };
+					Options.ApiVersion = 2;
+					UE_EOS_CHECK_API_MISMATCH(EOS_CONNECT_LOGIN_API_LATEST, 2);
+					Options.Credentials = &Credentials;
 
 #if ADD_USER_LOGIN_INFO
-				EOS_Connect_UserLoginInfo UserLoginInfo = {};
-				UserLoginInfo.ApiVersion = 2;
-				UE_EOS_CHECK_API_MISMATCH(EOS_CONNECT_USERLOGININFO_API_LATEST, 2);
-				const FTCHARToUTF8 DisplayNameUtf8(*GetPlatformDisplayName(LocalUserNum));
-				UserLoginInfo.DisplayName = DisplayNameUtf8.Get();
-				UserLoginInfo.NsaIdToken = nullptr; // Not set when EAS is not used
-
-				Options.UserLoginInfo = &UserLoginInfo;
+					EOS_Connect_UserLoginInfo UserLoginInfo = {};
+					UserLoginInfo.ApiVersion = 2;
+					UE_EOS_CHECK_API_MISMATCH(EOS_CONNECT_USERLOGININFO_API_LATEST, 2);
+					const FTCHARToUTF8 DisplayNameUtf8(*GetPlatformDisplayName(LocalUserNum));
+					UserLoginInfo.DisplayName = DisplayNameUtf8.Get();
+					UserLoginInfo.NsaIdToken = nullptr; // Not set when EAS is not used
+
+					Options.UserLoginInfo = &UserLoginInfo;
 #endif
-				FConnectLoginCallback* CallbackObj = new FConnectLoginCallback(AsWeak());
-				CallbackObj->CallbackLambda = [this, LocalUserNum](const EOS_Connect_LoginCallbackInfo* Data)
-					{
-						if (Data->ResultCode != EOS_EResult::EOS_Success)
+					FConnectLoginCallback* CallbackObj = new FConnectLoginCallback(AsWeak());
+					CallbackObj->CallbackLambda = [this, LocalUserNum](const EOS_Connect_LoginCallbackInfo* Data)
 						{
-							UE_LOG_ONLINE(Error, TEXT("Failed to refresh ConnectLogin for user %d. Failed with EOS result code (%s)"), LocalUserNum, *LexToString(Data->ResultCode));
+							if (Data->ResultCode != EOS_EResult::EOS_Success)
+							{
+								UE_LOG_ONLINE(Error, TEXT("Failed to refresh ConnectLogin for user %d. Failed with EOS result code (%s)"), LocalUserNum, *LexToString(Data->ResultCode));
 
-							TriggerOnAuthAboutToExpireDelegates(LocalUserNum);
-						}
-					};
+								TriggerOnAuthAboutToExpireDelegates(LocalUserNum);
+							}
+						};
 
-				EOS_Connect_Login(EOSSubsystem->ConnectHandle, &Options, CallbackObj, CallbackObj->GetCallbackPtr());
-			}
-			else // If not, we'll try to grab the token automatically from the platform
-			{
-				GetPlatformAuthToken(LocalUserNum,
-					FOnGetLinkedAccountAuthTokenCompleteDelegate::CreateLambda([this, WeakThis = AsWeak()](int32 LocalUserNum, bool bWasSuccessful, const FExternalAuthToken& AuthToken)
-					{
-						if (FUserManagerEOSPtr StrongThis = WeakThis.Pin())
-						{
-							if (!bWasSuccessful || !AuthToken.IsValid())
+					EOS_Connect_Login(EOSSubsystem->ConnectHandle, &Options, CallbackObj, CallbackObj->GetCallbackPtr());
+				}
+				else // If not, we'll try to grab the token automatically from the platform
+				{
+					GetPlatformAuthToken(LocalUserNum,
+						FOnGetLinkedAccountAuthTokenCompleteDelegate::CreateLambda([this, WeakThis = AsWeak()](int32 LocalUserNum, bool bWasSuccessful, const FExternalAuthToken& AuthToken)
 							{
-								UE_LOG_ONLINE(Error, TEXT("AutoRefreshConnectLogin(%d) failed due to the platform OSS giving an empty auth token"), LocalUserNum);
-								Logout(LocalUserNum);
-								return;
-							}
+								if (FUserManagerEOSPtr StrongThis = WeakThis.Pin())
+								{
+									if (!bWasSuccessful || !AuthToken.IsValid())
+									{
+										UE_LOG_ONLINE(Error, TEXT("AutoRefreshConnectLogin(%d) failed due to the platform OSS giving an empty auth token"), LocalUserNum);
+										Logout(LocalUserNum);
+										return;
+									}
 
-							// Now login into our EOS account
+									// Now login into our EOS account
 
-							IOnlineSubsystem* PlatformOSS = GetPlatformOSS();
-							FEOSSettings Settings = UEOSSettings::GetSettings();
-							EOS_EExternalCredentialType CredType = ToEOS_EExternalCredentialType(PlatformOSS ? PlatformOSS->GetSubsystemName() : EOSSubsystem->GetSubsystemName());
-							FConnectCredentials Credentials(CredType, ConvertExternalAuthTokenToFString(AuthToken));
+									IOnlineSubsystem* PlatformOSS = GetPlatformOSS();
+									FEOSSettings Settings = UEOSSettings::GetSettings();
+									EOS_EExternalCredentialType CredType = ToEOS_EExternalCredentialType(PlatformOSS ? PlatformOSS->GetSubsystemName() : EOSSubsystem->GetSubsystemName());
+									FConnectCredentials Credentials(CredType, ConvertExternalAuthTokenToFString(AuthToken));
 
-							EOS_Connect_LoginOptions Options = { };
-							Options.ApiVersion = 2;
-							UE_EOS_CHECK_API_MISMATCH(EOS_CONNECT_LOGIN_API_LATEST, 2);
-							Options.Credentials = &Credentials;
+									EOS_Connect_LoginOptions Options = { };
+									Options.ApiVersion = 2;
+									UE_EOS_CHECK_API_MISMATCH(EOS_CONNECT_LOGIN_API_LATEST, 2);
+									Options.Credentials = &Credentials;
 
 #if ADD_USER_LOGIN_INFO
-							EOS_Connect_UserLoginInfo UserLoginInfo = {};
-							UserLoginInfo.ApiVersion = 2;
-							UE_EOS_CHECK_API_MISMATCH(EOS_CONNECT_USERLOGININFO_API_LATEST, 2);
-							const FTCHARToUTF8 DisplayNameUtf8(*GetPlatformDisplayName(LocalUserNum));
-							UserLoginInfo.DisplayName = DisplayNameUtf8.Get();
-							UserLoginInfo.NsaIdToken = nullptr; // Not set when EAS is not used
-
-							Options.UserLoginInfo = &UserLoginInfo;
+									EOS_Connect_UserLoginInfo UserLoginInfo = {};
+									UserLoginInfo.ApiVersion = 2;
+									UE_EOS_CHECK_API_MISMATCH(EOS_CONNECT_USERLOGININFO_API_LATEST, 2);
+									const FTCHARToUTF8 DisplayNameUtf8(*GetPlatformDisplayName(LocalUserNum));
+									UserLoginInfo.DisplayName = DisplayNameUtf8.Get();
+									UserLoginInfo.NsaIdToken = nullptr; // Not set when EAS is not used
+
+									Options.UserLoginInfo = &UserLoginInfo;
 #endif
-							FConnectLoginCallback* CallbackObj = new FConnectLoginCallback(AsWeak());
-							CallbackObj->CallbackLambda = [this, LocalUserNum](const EOS_Connect_LoginCallbackInfo* Data)
-							{
-								if (Data->ResultCode != EOS_EResult::EOS_Success)
-								{
-									UE_LOG_ONLINE(Error, TEXT("Failed to refresh ConnectLogin(%d) failed with EOS result code (%s)"), LocalUserNum, *LexToString(Data->ResultCode));
-									Logout(LocalUserNum);
+									FConnectLoginCallback* CallbackObj = new FConnectLoginCallback(AsWeak());
+									CallbackObj->CallbackLambda = [this, LocalUserNum](const EOS_Connect_LoginCallbackInfo* Data)
+										{
+											if (Data->ResultCode != EOS_EResult::EOS_Success)
+											{
+												UE_LOG_ONLINE(Error, TEXT("Failed to refresh ConnectLogin(%d) failed with EOS result code (%s)"), LocalUserNum, *LexToString(Data->ResultCode));
+												Logout(LocalUserNum);
+											}
+										};
+									EOS_Connect_Login(EOSSubsystem->ConnectHandle, &Options, CallbackObj, CallbackObj->GetCallbackPtr());
 								}
-							};
-							EOS_Connect_Login(EOSSubsystem->ConnectHandle, &Options, CallbackObj, CallbackObj->GetCallbackPtr());
-						}
-					}));
+							}));
+				}
 			}
 		}
 	}
@@ -1300,64 +1461,8 @@ void FUserManagerEOS::CreateConnectedLogin(int32 LocalUserNum, EOS_EpicAccountId
 	EOS_Connect_CreateUser(EOSSubsystem->ConnectHandle, &Options, CallbackObj, CallbackObj->GetCallbackPtr());
 }
 
-typedef TEOSGlobalCallback<EOS_Connect_OnAuthExpirationCallback, EOS_Connect_AuthExpirationCallbackInfo, FUserManagerEOS> FRefreshAuthCallback;
-typedef TEOSGlobalCallback<EOS_Presence_OnPresenceChangedCallback, EOS_Presence_PresenceChangedCallbackInfo, FUserManagerEOS> FPresenceChangedCallback;
-typedef TEOSGlobalCallback<EOS_Friends_OnFriendsUpdateCallback, EOS_Friends_OnFriendsUpdateInfo, FUserManagerEOS> FFriendsStatusUpdateCallback;
-typedef TEOSGlobalCallback<EOS_Auth_OnLoginStatusChangedCallback, EOS_Auth_LoginStatusChangedCallbackInfo, FUserManagerEOS> FLoginStatusChangedCallback;
-
 void FUserManagerEOS::FullLoginCallback(int32 LocalUserNum, EOS_EpicAccountId AccountId, EOS_ProductUserId UserId)
 {
-	// Add our login status changed callback if not already set
-	if (LoginNotificationId == 0)
-	{
-		FLoginStatusChangedCallback* CallbackObj = new FLoginStatusChangedCallback(AsWeak());
-		LoginNotificationCallback = CallbackObj;
-		CallbackObj->CallbackLambda = [this](const EOS_Auth_LoginStatusChangedCallbackInfo* Data)
-			{
-				LoginStatusChanged(Data);
-			};
-
-		EOS_Auth_AddNotifyLoginStatusChangedOptions Options = { };
-		Options.ApiVersion = 1;
-		UE_EOS_CHECK_API_MISMATCH(EOS_AUTH_ADDNOTIFYLOGINSTATUSCHANGED_API_LATEST, 1);
-		LoginNotificationId = EOS_Auth_AddNotifyLoginStatusChanged(EOSSubsystem->AuthHandle, &Options, CallbackObj, CallbackObj->GetCallbackPtr());
-	}
-	// Register for friends updates if not set yet
-	if (FriendsNotificationId == 0)
-	{
-		FFriendsStatusUpdateCallback* CallbackObj = new FFriendsStatusUpdateCallback(AsWeak());
-		FriendsNotificationCallback = CallbackObj;
-		CallbackObj->CallbackLambda = [LocalUserNum, this](const EOS_Friends_OnFriendsUpdateInfo* Data)
-			{
-				FriendStatusChanged(Data);
-			};
-
-		EOS_Friends_AddNotifyFriendsUpdateOptions Options = { };
-		Options.ApiVersion = 1;
-		UE_EOS_CHECK_API_MISMATCH(EOS_FRIENDS_ADDNOTIFYFRIENDSUPDATE_API_LATEST, 1);
-		FriendsNotificationId = EOS_Friends_AddNotifyFriendsUpdate(EOSSubsystem->FriendsHandle, &Options, CallbackObj, CallbackObj->GetCallbackPtr());
-	}
-	// Register for presence updates if not set yet
-	if (PresenceNotificationId == 0)
-	{
-		FPresenceChangedCallback* CallbackObj = new FPresenceChangedCallback(AsWeak());
-		PresenceNotificationCallback = CallbackObj;
-		CallbackObj->CallbackLambda = [LocalUserNum, this](const EOS_Presence_PresenceChangedCallbackInfo* Data)
-			{
-				if (FUniqueNetIdEOSRegistry::Find(Data->PresenceUserId))
-				{
-					// Update the presence data to the most recent
-					UpdatePresence(LocalUserNum, Data->PresenceUserId);
-					return;
-				}
-			};
-
-		EOS_Presence_AddNotifyOnPresenceChangedOptions Options = { };
-		Options.ApiVersion = 1;
-		UE_EOS_CHECK_API_MISMATCH(EOS_PRESENCE_ADDNOTIFYONPRESENCECHANGED_API_LATEST, 1);
-		PresenceNotificationId = EOS_Presence_AddNotifyOnPresenceChanged(EOSSubsystem->PresenceHandle, &Options, CallbackObj, CallbackObj->GetCallbackPtr());
-	}
-
 	// Set the ProductUserId mapping if the platform supports it
 	EOSSubsystem->GetEOSHelpers()->AddProductUserIdMapping(EOSSubsystem->GetEOSPlatformHandle(), UserId, LocalUserNum);
 
@@ -1368,9 +1473,6 @@ void FUserManagerEOS::FullLoginCallback(int32 LocalUserNum, EOS_EpicAccountId Ac
 	TriggerOnLoginStatusChangedDelegates(LocalUserNum, ELoginStatus::NotLoggedIn, ELoginStatus::LoggedIn, *LocalUser.UniqueNetId);
 }
 
-typedef TEOSCallback<EOS_Auth_OnLogoutCallback, EOS_Auth_LogoutCallbackInfo, FUserManagerEOS> FAuthLogoutCallback;
-typedef TEOSCallback<EOS_Connect_OnLogoutCallback, EOS_Connect_LogoutCallbackInfo, FUserManagerEOS> FConnectLogoutCallback;
-
 bool FUserManagerEOS::Logout(int32 LocalUserNum)
 {
 	FUniqueNetIdEOSPtr UserId = GetLocalUniqueNetIdEOS(LocalUserNum);
@@ -1381,6 +1483,8 @@ bool FUserManagerEOS::Logout(int32 LocalUserNum)
 		TriggerOnLogoutCompleteDelegates(LocalUserNum, false);
 		return false;
 	}
+	FLocalUserEOS& LocalUser = GetLocalUserChecked(LocalUserNum);
+	LocalUser.bIsLoggingOut = true;
 
 	FEOSSettings Settings = UEOSSettings::GetSettings();
 
@@ -1419,7 +1523,9 @@ void FUserManagerEOS::CallEOSAuthLogout(int32 LocalUserNum, FUniqueNetIdEOSPtr U
 			else
 			{
 				// EOS_Auth_DeletePersistentAuth and EOS_Connect_Logout are not called on failure. The game should attempt to logout again
-				UE_LOG_ONLINE(Warning, TEXT("EOS_Auth_Logout Result=[%s]"), *LexToString(Data->ResultCode))
+				UE_LOG_ONLINE(Warning, TEXT("EOS_Auth_Logout Result=[%s]"), *LexToString(Data->ResultCode));
+				FLocalUserEOS& LocalUser = GetLocalUserChecked(LocalUserNum);
+				LocalUser.bIsLoggingOut = false;
 				TriggerOnLogoutCompleteDelegates(LocalUserNum, false);
 			}
 		};
@@ -1438,12 +1544,12 @@ void FUserManagerEOS::CallEOSConnectLogout(int32 LocalUserNum, FUniqueNetIdEOSPt
 	CallbackObj->CallbackLambda = [LocalUserNum, this](const EOS_Connect_LogoutCallbackInfo* Data)
 		{
 			const bool bWasConnectLogoutSuccessful = Data->ResultCode == EOS_EResult::EOS_Success;
-
 			UE_CLOG_ONLINE(!bWasConnectLogoutSuccessful, Warning, TEXT("EOS_Connect_Logout Result=[%s]"), *LexToString(Data->ResultCode));
 
-			if (bWasConnectLogoutSuccessful)
+			if (IsLocalUserValid(LocalUserNum))
 			{
-				RemoveLocalUser(LocalUserNum);
+				FLocalUserEOS& LocalUser = GetLocalUserChecked(LocalUserNum);
+				LocalUser.bIsLoggingOut = false;
 			}
 
 			TriggerOnLogoutCompleteDelegates(LocalUserNum, bWasConnectLogoutSuccessful);
@@ -1538,25 +1644,6 @@ FLocalUserEOS& FUserManagerEOS::UpdateLocalUser(int32 LocalUserNum, EOS_EpicAcco
 	FUserOnlineAccountEOSRef UserAccountRef(new FUserOnlineAccountEOS(UserNetId, *EOSSubsystem));
 	LocalUser.UserOnlineAccount = UserAccountRef;
 
-	// Add auth refresh notification if not set for this user yet
-	if (!LocalUser.ConnectLoginNotification.IsValid())
-	{
-		FNotificationIdCallbackPairPtr NotificationPairPtr = MakeShared<FNotificationIdCallbackPair>();
-		LocalUser.ConnectLoginNotification = NotificationPairPtr;
-
-		FRefreshAuthCallback* CallbackObj = new FRefreshAuthCallback(AsWeak());
-		NotificationPairPtr->Callback = CallbackObj;
-		CallbackObj->CallbackLambda = [LocalUserNum, this](const EOS_Connect_AuthExpirationCallbackInfo* Data)
-			{
-				AutoRefreshConnectLogin(LocalUserNum);
-			};
-
-		EOS_Connect_AddNotifyAuthExpirationOptions Options = { };
-		Options.ApiVersion = 1;
-		UE_EOS_CHECK_API_MISMATCH(EOS_CONNECT_ADDNOTIFYAUTHEXPIRATION_API_LATEST, 1);
-		NotificationPairPtr->NotificationId = EOS_Connect_AddNotifyAuthExpiration(EOSSubsystem->ConnectHandle, &Options, CallbackObj, CallbackObj->GetCallbackPtr());
-	}
-
 	const FOnlineUserEOSRef UserRef(new FOnlineUserEOS(UserNetId, *EOSSubsystem));
 	UniqueNetIdToUserRefMap.Emplace(UserNetId, UserRef);
 
@@ -1741,9 +1828,12 @@ int32 FUserManagerEOS::GetLocalUserNumFromEpicAccountId(const EOS_EpicAccountId&
 		{
 			const FUniqueNetIdEOSPtr& Entry = LocalUsers[LocalUserNum].UniqueNetId;
 
-			if (Entry->GetEpicAccountId() == EpicAccountId)
+			if (Entry.IsValid())
 			{
-				return LocalUserNum;
+				if (Entry->GetEpicAccountId() == EpicAccountId)
+				{
+					return LocalUserNum;
+				}
 			}
 		}
 	}
@@ -1758,10 +1848,13 @@ int32 FUserManagerEOS::GetLocalUserNumFromProductUserId(const EOS_ProductUserId&
 		if (LocalUsers.IsValidIndex(LocalUserNum))
 		{
 			const FUniqueNetIdEOSPtr& Entry = LocalUsers[LocalUserNum].UniqueNetId;
-
-			if (Entry->GetProductUserId() == ProductUserId)
+			
+			if (Entry.IsValid())
 			{
-				return LocalUserNum;
+				if (Entry->GetProductUserId() == ProductUserId)
+				{
+					return LocalUserNum;
+				}
 			}
 		}
 	}
@@ -1777,9 +1870,12 @@ const FUniqueNetIdEOSPtr FUserManagerEOS::GetLocalUniqueNetIdEOS(const EOS_EpicA
 		{
 			const FUniqueNetIdEOSPtr& Entry = LocalUsers[LocalUserNum].UniqueNetId;
 
-			if (Entry->GetEpicAccountId() == EpicAccountId)
+			if (Entry.IsValid())
 			{
-				return Entry;
+				if (Entry->GetEpicAccountId() == EpicAccountId)
+				{
+					return Entry;
+				}
 			}
 		}
 	}
@@ -1795,9 +1891,12 @@ const FUniqueNetIdEOSPtr FUserManagerEOS::GetLocalUniqueNetIdEOS(const EOS_Produ
 		{
 			const FUniqueNetIdEOSPtr& Entry = LocalUsers[LocalUserNum].UniqueNetId;
 
-			if (Entry->GetProductUserId() == ProductUserId)
+			if (Entry.IsValid())
 			{
-				return Entry;
+				if (Entry->GetProductUserId() == ProductUserId)
+				{
+					return Entry;
+				}
 			}
 		}
 	}
@@ -3380,6 +3479,19 @@ void FUserManagerEOS::QueryPresence(const FUniqueNetId& UserId, const FOnPresenc
 	Delegate.ExecuteIfBound(UserId, true);
 }
 
+void FUserManagerEOS::PresenceChanged(const EOS_Presence_PresenceChangedCallbackInfo* Data)
+{
+	const int32 LocalUserNum = GetLocalUserNumFromEpicAccountId(Data->LocalUserId);
+
+	if (IsLocalUserValid(LocalUserNum))
+	{
+		if (FUniqueNetIdEOSRegistry::Find(Data->PresenceUserId))
+		{
+			UpdatePresence(LocalUserNum, Data->PresenceUserId);
+		}
+	}
+}
+
 void FUserManagerEOS::UpdatePresence(int32 LocalUserNum, EOS_EpicAccountId AccountId)
 {
 	EOS_Presence_Info* PresenceInfo = nullptr;
diff --git a/Engine/Plugins/Online/OnlineSubsystemEOS/Source/OnlineSubsystemEOS/Private/UserManagerEOS.h b/Engine/Plugins/Online/OnlineSubsystemEOS/Source/OnlineSubsystemEOS/Private/UserManagerEOS.h
--- a/Engine/Plugins/Online/OnlineSubsystemEOS/Source/OnlineSubsystemEOS/Private/UserManagerEOS.h
+++ b/Engine/Plugins/Online/OnlineSubsystemEOS/Source/OnlineSubsystemEOS/Private/UserManagerEOS.h
@@ -171,7 +171,6 @@ struct FNotificationIdCallbackPair
 		delete Callback;
 	}
 };
-typedef TSharedPtr<FNotificationIdCallbackPair> FNotificationIdCallbackPairPtr;
 
 /** Cache for the info passed on to ReadFriendsList, kept while user info queries complete */
 struct ReadUserListInfo
@@ -234,11 +233,11 @@ struct FLocalUserEOS
 	/** Last Login Credentials used for a login attempt */
 	TSharedPtr<FOnlineAccountCredentials> LastLoginCredentials_Legacy;
 
-	/** Struct containing the information related to the connect login notification for the user */
-	FNotificationIdCallbackPairPtr ConnectLoginNotification;
-
 	/** If set, Connect Login will be re-attempted with it. If this attempt fails, a new token will need to be provided */
 	TOptional<FString> UserProvidedConnectAuthToken;
+
+	/** Track if the player is in the process of logging out to avoid calling EOS SDK logout APIs multiple times */
+	bool bIsLoggingOut = false;
 };
 
 /**
@@ -388,13 +387,15 @@ public:
 	void LoginViaExternalAuth(int32 LocalUserNum);
 	void CreateConnectedLogin(int32 LocalUserNum, EOS_EpicAccountId AccountId, EOS_ContinuanceToken Token);
 	void LinkEAS(int32 LocalUserNum, EOS_ContinuanceToken Token);
-	void AutoRefreshConnectLogin(int32 LocalUserNum);
+	void AutoRefreshConnectLogin(const EOS_Connect_AuthExpirationCallbackInfo* Data);
 	bool ConnectLoginNoEAS(int32 LocalUserNum);
 
 	void FullLoginCallback(int32 LocalUserNum, EOS_EpicAccountId AccountId, EOS_ProductUserId UserId);
 	void FriendStatusChanged(const EOS_Friends_OnFriendsUpdateInfo* Data);
 	void FriendStatusChangedImpl(EOS_EpicAccountId LocalUserId, EOS_EpicAccountId TargetUserId, EOS_EFriendsStatus PreviousStatus, EOS_EFriendsStatus CurrentStatus);
-	void LoginStatusChanged(const EOS_Auth_LoginStatusChangedCallbackInfo* Data);
+	void AuthLoginStatusChanged(const EOS_Auth_LoginStatusChangedCallbackInfo* Data);
+	void ConnectLoginStatusChanged(const EOS_Connect_LoginStatusChangedCallbackInfo* Data);
+	void PresenceChanged(const EOS_Presence_PresenceChangedCallbackInfo* Data);
 
 	int32 GetDefaultLocalUser() const { return DefaultLocalUser; }
 
@@ -445,8 +446,12 @@ private:
 	int32 DefaultLocalUser;
 
 	/** Notification state for SDK events */
-	EOS_NotificationId LoginNotificationId = EOS_INVALID_NOTIFICATIONID;
-	FCallbackBase* LoginNotificationCallback = nullptr;
+	EOS_NotificationId AuthLoginStatusChangedNotificationId = EOS_INVALID_NOTIFICATIONID;
+	FCallbackBase* AuthLoginStatusChangedNotificationCallback = nullptr;
+	EOS_NotificationId ConnectLoginStatusChangedNotificationId = EOS_INVALID_NOTIFICATIONID;
+	FCallbackBase* ConnectLoginStatusChangedNotificationCallback = nullptr;
+	EOS_NotificationId ConnectLoginExpiredNotificationId = EOS_INVALID_NOTIFICATIONID;
+	FCallbackBase* ConnectLoginExpiredNotificationCallback = nullptr;
 	EOS_NotificationId FriendsNotificationId = EOS_INVALID_NOTIFICATIONID;
 	FCallbackBase* FriendsNotificationCallback = nullptr;
 	EOS_NotificationId PresenceNotificationId = EOS_INVALID_NOTIFICATIONID;
