UE: EOS ReadLeaderboards - Platform Updates fixes

#Synced-CL 52126719

diff --git a/Engine/Plugins/Online/OnlineSubsystemEOS/Source/OnlineSubsystemEOS/Private/OnlineLeaderboardsEOS.cpp b/Engine/Plugins/Online/OnlineSubsystemEOS/Source/OnlineSubsystemEOS/Private/OnlineLeaderboardsEOS.cpp
--- a/Engine/Plugins/Online/OnlineSubsystemEOS/Source/OnlineSubsystemEOS/Private/OnlineLeaderboardsEOS.cpp
+++ b/Engine/Plugins/Online/OnlineSubsystemEOS/Source/OnlineSubsystemEOS/Private/OnlineLeaderboardsEOS.cpp
@@ -12,6 +12,121 @@
 #include "eos_leaderboards.h"
 #include "eos_userinfo.h"
 
+namespace Private
+{
+struct FRankedUserInfo
+{
+	FUniqueNetIdEOSRef NetId;
+	EOS_ProductUserId ProductUserId;
+	FString Nickname;
+	uint32 Rank;
+	int32 Score;
+
+	FRankedUserInfo(const FUniqueNetIdEOSRef InNetId, EOS_ProductUserId InProductUserId, FString InNickname, uint32 InRank, int32 InScore)
+		: NetId(InNetId)
+		, ProductUserId(InProductUserId)
+		, Nickname(MoveTemp(InNickname))
+		, Rank(InRank)
+		, Score(InScore)
+	{
+	}
+};
+
+/** Collect rank records into FRankedUserInfo array, resolving against known NetIds */
+void CollectRankedUsers(EOS_HLeaderboards LeaderboardsHandle, uint32 StartIndex, uint32 EndIndex,
+	const TMap<EOS_ProductUserId, FUniqueNetIdEOSRef>& ResolvedUniqueNetIds,
+	TArray<FRankedUserInfo>& OutRankedUsers, TArray<EOS_ProductUserId>& OutProductUserIds)
+{
+	EOS_Leaderboards_CopyLeaderboardRecordByIndexOptions CopyOptions = { };
+	CopyOptions.ApiVersion = 2;
+	UE_EOS_CHECK_API_MISMATCH(EOS_LEADERBOARDS_COPYLEADERBOARDRECORDBYINDEX_API_LATEST, 2);
+
+	for (uint32 Index = StartIndex; Index <= EndIndex; Index++)
+	{
+		CopyOptions.LeaderboardRecordIndex = Index;
+
+		EOS_Leaderboards_LeaderboardRecord* Record = nullptr;
+		EOS_EResult Result = EOS_Leaderboards_CopyLeaderboardRecordByIndex(LeaderboardsHandle, &CopyOptions, &Record);
+		if (Result == EOS_EResult::EOS_Success)
+		{
+			if (const FUniqueNetIdEOSRef* NetId = ResolvedUniqueNetIds.Find(Record->UserId))
+			{
+				FString Nickname = UTF8_TO_TCHAR(Record->UserDisplayName);
+				if (Nickname.IsEmpty())
+				{
+					Nickname = TEXT("Unknown Player");
+				}
+				OutRankedUsers.Emplace(*NetId, Record->UserId, MoveTemp(Nickname), Record->Rank, Record->Score);
+				OutProductUserIds.Add(Record->UserId);
+			}
+
+			EOS_Leaderboards_LeaderboardRecord_Release(Record);
+		}
+	}
+}
+
+/** Populate ReadObject rows from RankedUsers with all columns fetched via CopyLeaderboardUserScoreByUserId */
+void PopulateRowsWithUserScores(EOS_HLeaderboards LeaderboardsHandle, 
+	const TArray<FRankedUserInfo>& RankedUsers, 
+	const FOnlineLeaderboardReadRef& ReadObject)
+{
+	char StatName[EOS_OSS_STRING_BUFFER_LENGTH];
+	for (const FRankedUserInfo& Info : RankedUsers)
+	{
+		FOnlineStatsRow& Row = ReadObject->Rows.Emplace_GetRef(Info.Nickname, Info.NetId);
+		Row.Rank = Info.Rank;
+
+		EOS_Leaderboards_CopyLeaderboardUserScoreByUserIdOptions UserCopyOptions = { };
+		UserCopyOptions.ApiVersion = 1;
+		UE_EOS_CHECK_API_MISMATCH(EOS_LEADERBOARDS_COPYLEADERBOARDUSERSCOREBYUSERID_API_LATEST, 1);
+		UserCopyOptions.UserId = Info.ProductUserId;
+		UserCopyOptions.StatName = StatName;
+
+		for (const FColumnMetaData& Column : ReadObject->ColumnMetadata)
+		{
+			FCStringAnsi::Strncpy(StatName, TCHAR_TO_UTF8(*Column.ColumnName), EOS_OSS_STRING_BUFFER_LENGTH);
+			EOS_Leaderboards_LeaderboardUserScore* UserScore = nullptr;
+			EOS_EResult CopyResult = EOS_Leaderboards_CopyLeaderboardUserScoreByUserId(LeaderboardsHandle, &UserCopyOptions, &UserScore);
+			if (CopyResult == EOS_EResult::EOS_Success)
+			{
+				Row.Columns.Emplace(Column.ColumnName, FVariantData(UserScore->Score));
+				EOS_Leaderboards_LeaderboardUserScore_Release(UserScore);
+			}
+			else
+			{
+				Row.Columns.Emplace(Column.ColumnName, FVariantData());
+			}
+		}
+	}
+}
+
+/** Populate ReadObject rows with empty columns (fallback when user scores query fails) */
+void PopulateRowsWithEmptyColumns(const TArray<FRankedUserInfo>& RankedUsers, const FOnlineLeaderboardReadRef& ReadObject)
+{
+	for (const FRankedUserInfo& Info : RankedUsers)
+	{
+		FOnlineStatsRow& Row = ReadObject->Rows.Emplace_GetRef(Info.Nickname, Info.NetId);
+		Row.Rank = Info.Rank;
+		for (const FColumnMetaData& Column : ReadObject->ColumnMetadata)
+		{
+			Row.Columns.Emplace(Column.ColumnName, FVariantData());
+		}
+	}
+}
+
+/** Populate ReadObject rows with only the sorted column, using scores already collected in FRankedUserInfo */
+void PopulateRowsSortedColumnOnly(const TArray<FRankedUserInfo>& RankedUsers, const FOnlineLeaderboardReadRef& ReadObject)
+{
+	for (const FRankedUserInfo& Info : RankedUsers)
+	{
+		FOnlineStatsRow& Row = ReadObject->Rows.Emplace_GetRef(Info.Nickname, Info.NetId);
+		Row.Rank = Info.Rank;
+		Row.Columns.Add(ReadObject->SortedColumn, FVariantData(Info.Score));
+	}
+}
+
+/* Private */ }
+
 struct FQueryLeaderboardForUserOptions :
 	public EOS_Leaderboards_QueryLeaderboardUserScoresOptions
 {
@@ -297,42 +412,72 @@ bool FOnlineLeaderboardsEOS::ReadLeaderboardsAroundRank(int32 Rank, uint32 Range
 			if (Result == EOS_EResult::EOS_Success)
 			{
 				UsersToResolve.Add(Record->UserId);
+				EOS_Leaderboards_LeaderboardRecord_Release(Record);
 			}
 		}
 
-		EOSSubsystem->UserManager->ResolveUniqueNetIds(0, UsersToResolve, [this, StartIndex, NewEndIndex, ReadObject](TMap<EOS_ProductUserId, FUniqueNetIdEOSRef> ResolvedUniqueNetIds, const FOnlineError& Error) mutable
+		EOSSubsystem->UserManager->ResolveUniqueNetIds(0, UsersToResolve, [this, StartIndex, NewEndIndex, ReadObject]
+		(const TMap<EOS_ProductUserId, FUniqueNetIdEOSRef>& ResolvedUniqueNetIds, const FOnlineError& Error)
+		{
+			TArray<Private::FRankedUserInfo> RankedUsers;
+			TArray<EOS_ProductUserId> ProductUserIds;
+			Private::CollectRankedUsers(EOSSubsystem->LeaderboardsHandle, StartIndex, NewEndIndex, ResolvedUniqueNetIds, RankedUsers, ProductUserIds);
+
+			if (RankedUsers.IsEmpty())
 			{
-				// Then we copy the rest of the data
-				EOS_Leaderboards_CopyLeaderboardRecordByIndexOptions CopyOptions = { };
-				CopyOptions.ApiVersion = 2;
-				UE_EOS_CHECK_API_MISMATCH(EOS_LEADERBOARDS_COPYLEADERBOARDRECORDBYINDEX_API_LATEST, 2);
+				ReadObject->ReadState = EOnlineAsyncTaskState::Done;
+				TriggerOnLeaderboardReadCompleteDelegates(true);
+				return;
+			}
 
-				for (uint32 Index = StartIndex; Index <= NewEndIndex; Index++)
+			// if there are extra columns beyond the sorted column, issue a follow-up QueryLeaderboardUserScores to populate all requested columns:
+			const uint32 NumColumns = ReadObject->ColumnMetadata.Num();
+			if (NumColumns > 1 || (NumColumns == 1 && ReadObject->ColumnMetadata[0].ColumnName != ReadObject->SortedColumn))
+			{
+				// TSharedPtr ensures cleanup even if TEOSCallback skips lambda execution (e.g., Owner destroyed)
+				TSharedPtr<FQueryLeaderboardForUserOptions> UserScoresOptions = MakeShared<FQueryLeaderboardForUserOptions>(ReadObject->ColumnMetadata.Num(), ProductUserIds);
+				UserScoresOptions->LocalUserId = EOSSubsystem->UserManager->GetLocalProductUserId(0);
+				int32 ColIdx = 0;
+				for (const FColumnMetaData& Column : ReadObject->ColumnMetadata)
 				{
-					CopyOptions.LeaderboardRecordIndex = Index;
+					FCStringAnsi::Strncpy(UserScoresOptions->PointerArray[ColIdx], TCHAR_TO_UTF8(*Column.ColumnName), EOS_OSS_STRING_BUFFER_LENGTH);
+					EOS_Leaderboards_UserScoresQueryStatInfo& StatInfo = UserScoresOptions->StatInfoArray[ColIdx];
+					StatInfo.ApiVersion = 1;
+					UE_EOS_CHECK_API_MISMATCH(EOS_LEADERBOARDS_USERSCORESQUERYSTATINFO_API_LATEST, 1);
+					StatInfo.StatName = UserScoresOptions->PointerArray[ColIdx];
+					// Aggregation matches ReadLeaderboards: if a leaderboard uses a different aggregation
+					// (Min, Max, Sum), these values may differ from what QueryLeaderboardRanks used for ranking
+					StatInfo.Aggregation = EOS_ELeaderboardAggregation::EOS_LA_Latest;
+					++ColIdx;
+				}
 
-					EOS_Leaderboards_LeaderboardRecord* Record = nullptr;
-					EOS_EResult Result = EOS_Leaderboards_CopyLeaderboardRecordByIndex(EOSSubsystem->LeaderboardsHandle, &CopyOptions, &Record);
-					if (Result == EOS_EResult::EOS_Success)
+				FQueryLeaderboardForUsersCallback* UserScoresCallback = new FQueryLeaderboardForUsersCallback(FOnlineLeaderboardsEOSWeakPtr(AsShared()));
+				UserScoresCallback->CallbackLambda = [this, ReadObject, RankedUsers = MoveTemp(RankedUsers), UserScoresOptions]
+				(const EOS_Leaderboards_OnQueryLeaderboardUserScoresCompleteCallbackInfo* Data)
+				{
+					if (Data->ResultCode != EOS_EResult::EOS_Success)
 					{
-						FString Nickname = UTF8_TO_TCHAR(Record->UserDisplayName);
-						if (Nickname.IsEmpty())
-						{
-							Nickname = TEXT("Unknown Player");
-						}
-						
-						if (const FUniqueNetIdEOSRef* NetId = ResolvedUniqueNetIds.Find(Record->UserId))
-						{
-							FOnlineStatsRow* Row = new(ReadObject->Rows) FOnlineStatsRow(Nickname, *NetId);
-							Row->Rank = Record->Rank;
-							Row->Columns.Add(ReadObject->SortedColumn, FVariantData(Record->Score));
-						}
+						UE_LOG_ONLINE_LEADERBOARD(Warning, TEXT("ReadLeaderboardsAroundRank: follow-up QueryLeaderboardUserScores failed (%s), falling back to empty columns"), *LexToString(Data->ResultCode));
+						Private::PopulateRowsWithEmptyColumns(RankedUsers, ReadObject);
+						ReadObject->ReadState = EOnlineAsyncTaskState::Done;
+						TriggerOnLeaderboardReadCompleteDelegates(true);
+						return;
 					}
-				}
 
+					Private::PopulateRowsWithUserScores(EOSSubsystem->LeaderboardsHandle, RankedUsers, ReadObject);
+					ReadObject->ReadState = EOnlineAsyncTaskState::Done;
+					TriggerOnLeaderboardReadCompleteDelegates(true);
+				};
+
+				EOS_Leaderboards_QueryLeaderboardUserScores(EOSSubsystem->LeaderboardsHandle, UserScoresOptions.Get(), UserScoresCallback, UserScoresCallback->GetCallbackPtr());
+			}
+			else
+			{
+				Private::PopulateRowsSortedColumnOnly(RankedUsers, ReadObject);
 				ReadObject->ReadState = EOnlineAsyncTaskState::Done;
 				TriggerOnLeaderboardReadCompleteDelegates(true);
-			});
+			}
+		});
 	};
 
 	ReadObject->ReadState = EOnlineAsyncTaskState::InProgress;
@@ -444,7 +589,7 @@ bool FOnlineLeaderboardsEOS::HandleLeaderboardsExec(UWorld* InWorld, const TCHAR
 
 		ReadLeaderboardsForFriends(LocalUserNum, ReadRef);
 	}
-	else if (FParse::Command(&Cmd, TEXT("ReadLeaderboardsAroundRank"))) /* ONLINE LEADERBOARDS ReadLeaderboardsAroundRank Rank=1 Range=10 LeaderboardName=Deaths SortedColumn=Deaths DataType=Int32 */
+	else if (FParse::Command(&Cmd, TEXT("ReadLeaderboardsAroundRank"))) /* ONLINE LEADERBOARDS ReadLeaderboardsAroundRank Rank=1 Range=10 LeaderboardName=Deaths SortedColumn=Deaths DataType=Int32 ExtraColumns=Kills,Wins */
 	{
 		int Rank;
 		FParse::Value(Cmd, TEXT("Rank="), Rank);
@@ -464,11 +609,31 @@ bool FOnlineLeaderboardsEOS::HandleLeaderboardsExec(UWorld* InWorld, const TCHAR
 		ReadRef->SortedColumn = SortedColumn;
 		ReadRef->ColumnMetadata.Add(FColumnMetaData(SortedColumn, DataType));
 
+		// Parse optional extra columns (comma-separated) to test multi-column support
+		FString ExtraColumnsStr;
+		if (FParse::Value(Cmd, TEXT("ExtraColumns="), ExtraColumnsStr))
+		{
+			TArray<FString> ExtraColumns;
+			ExtraColumnsStr.ParseIntoArray(ExtraColumns, TEXT(","), true);
+			for (const FString& ExtraColumn : ExtraColumns)
+			{
+				ReadRef->ColumnMetadata.Add(FColumnMetaData(ExtraColumn, DataType));
+			}
+		}
+
 		AddOnLeaderboardReadCompleteDelegate_Handle(FOnLeaderboardReadCompleteDelegate::CreateLambda([ReadRef](bool bWasSuccessful)
+		{
+			UE_LOG_ONLINE_LEADERBOARD(Log, TEXT("FOnlineLeaderboardsEOS::ReadLeaderboardsAroundRank finished with bWasSuccessful=%s"), *LexToString(bWasSuccessful));
+			UE_LOG_ONLINE_LEADERBOARD(Log, TEXT("%s"), *ReadRef->ToLogString());
+			for (const FOnlineStatsRow& Row : ReadRef->Rows)
 			{
-				UE_LOG_ONLINE_LEADERBOARD(Log, TEXT("FOnlineLeaderboardsEOS::ReadLeaderboards finished with bWasSuccessful=%s"), *LexToString(bWasSuccessful));
-				UE_LOG_ONLINE_LEADERBOARD(Log, TEXT("%s"), *ReadRef->ToLogString());
-			}));
+				UE_LOG_ONLINE_LEADERBOARD(Log, TEXT("  Rank=%d Player=%s Columns=%d"), Row.Rank, *Row.NickName, Row.Columns.Num());
+				for (const TTuple<FString, FVariantData>& ColumnPair : Row.Columns)
+				{
+					UE_LOG_ONLINE_LEADERBOARD(Log, TEXT("    %s = %s"), *ColumnPair.Key, *ColumnPair.Value.ToString());
+				}
+			}
+		}));
 
 		ReadLeaderboardsAroundRank(Rank, Range, ReadRef);
 	}
